Make the -DataAgeThreshold parameter of BuildPatchTool.exe optional.

Default behavior is as follows:
* During compactify, ALL data is eligible for deletion
* During patch generation, ALL data is eligible for reuse

diff --git a/Engine/Source/Programs/AutomationTool/Scripts/CopyBuildToStagingDirectory.Automation.cs b/Engine/Source/Programs/AutomationTool/Scripts/CopyBuildToStagingDirectory.Automation.cs
--- a/Engine/Source/Programs/AutomationTool/Scripts/CopyBuildToStagingDirectory.Automation.cs
+++ b/Engine/Source/Programs/AutomationTool/Scripts/CopyBuildToStagingDirectory.Automation.cs
@@ -544,7 +544,6 @@ public partial class Project : CommandUtils
 				string VersionString = P4Branch + "-CL-" + P4Change;
 
 				string CmdLine = String.Format("-BuildRoot=\"{0}\" -CloudDir=\"{1}\" -AppID={2} -AppName=\"{3}\" -BuildVersion=\"{4}\" -AppLaunch=\"{5}\"", BuildRoot, CloudDir, AppID, AppName, VersionString, AppLaunch);
-				CmdLine += " -DataAgeThreshold=12";
 				CmdLine += " -AppArgs=\"\"";
 				CmdLine += " -custom=\"bIsPatch=false\"";
 				CmdLine += String.Format(" -customint=\"ChunkID={0}\"", ChunkID);
diff --git a/Engine/Source/Programs/BuildPatchTool/Private/BuildPatchToolMain.cpp b/Engine/Source/Programs/BuildPatchTool/Private/BuildPatchToolMain.cpp
--- a/Engine/Source/Programs/BuildPatchTool/Private/BuildPatchToolMain.cpp
+++ b/Engine/Source/Programs/BuildPatchTool/Private/BuildPatchToolMain.cpp
@@ -35,7 +35,8 @@
 	-PrereqArgs=""		Specifies in quotes, the commandline to send to prerequisites installer on launch
 	-DataAgeThreshold=12.5		Specified the maximum age (in days) of existing patch data which can be reused in the generated manifest
 
-	NB: If -DataAgeThreshold is not specified on the command line, then it *MUST* be present in the [PatchGeneration] section of BuildPatchTool.ini in the cloud directory
+	NB: -DataAgeThreshold can also be present in the [PatchGeneration] section of BuildPatchTool.ini in the cloud directory
+	NB: If -DataAgeThreshold is not supplied, either on the command-line or in BuildPatchTool.ini, then all existing data is eligible for reuse in the generated manifest
 
 	Example command lines:
 	-BuildRoot="D:\Builds\ExampleGame_[07-02_03.00]" -FileIgnoreList="D:\Builds\ExampleGame_[07-02_03.00]\Manifest_DebugFiles.txt" -CloudDir="D:\BuildPatchCloud" -AppID=123456 -AppName="Example" -BuildVersion="ExampleGame_[07-02_03.00]" -AppLaunch=".\ExampleGame\Binaries\Win32\ExampleGame.exe" -AppArgs="-pak -nosteam" -DataAgeThreshold=12
@@ -59,7 +60,8 @@
 	-DataAgeThreshold=14.25		The maximum age in days of chunk files that will be retained. All older chunks will be deleted.
 
 	NB: If -ManifestsList is specified, then -ManifestsFile is ignored.
-	NB: If -DataAgeThreshold is not specified on the command line, then it *MUST* be present in the [Compactify] section of BuildPatchTool.ini in the cloud directory
+	NB: -DataAgeThreshold can also be present in the [Compactify] section of BuildPatchTool.ini in the cloud directory
+	NB: If -DataAgeThreshold is not supplied, either on the command-line or in BuildPatchTool.ini, then all unreferenced existing data is eligible for deletion by the compactify process
 
 	Example command lines:
 	-CloudDir="E:\BuildPatchCloud" -compactify -DataAgeThreshold=14
@@ -151,6 +153,7 @@ int32 BuildPatchToolMain( const TCHAR* CommandLine )
 	bool bCompactify = false;
 	bool bPatchGeneration = true;
 	bool bPreview = false;
+	bool bPatchWithReuseAgeThreshold = true;
 
 	// Collect all the info from the CommandLine
 	TArray< FString > Tokens, Switches;
@@ -345,20 +348,23 @@ int32 BuildPatchToolMain( const TCHAR* CommandLine )
 		}
 		else if (bSuccess && bCompactify)
 		{
-			// For compactification, if we don't pass in DataAgeThreshold, then it *MUST* be specified in BuildPatchTool.ini
+			// For compactification, if we don't pass in DataAgeThreshold, and it's not in BuildPatchTool.ini,
+			// then we set it to zero, to indicate that any unused chunks are valid for deletion
 			if (!GConfig->GetFloat(TEXT("Compactify"), TEXT("DataAgeThreshold"), DataAgeThreshold, IniFile))
 			{
-				GLog->Log(ELogVerbosity::Error, TEXT("DataAgeThreshold must be specified either on the command line or in the Compactify section of BuildPatchTool.ini"));
-				bSuccess = false; // We couldn't get the DataAgeThreshold value from the .ini file
+				GLog->Log(ELogVerbosity::Warning, TEXT("DataAgeThreshold not supplied, so all unreferenced data is eliglble for deletion. Note that this process is NOT compatible with any concurrently running patch generaiton processes"));
+				DataAgeThreshold = 0.0f;
 			}
 		}
 		else if (bSuccess && bPatchGeneration)
 		{
-			// For patch generation, if we don't pass in DataAgeThreshold, then it *MUST* be specified in BuildPatchTool.ini
+			// For patch generation, if we don't pass in DataAgeThreshold, and it's not specified in BuildPatchTool.ini,
+			// then we set bChunkWithReuseAgeThreshold to false, which indicates that *all* patch data is valid for reuse
 			if (!GConfig->GetFloat(TEXT("PatchGeneration"), TEXT("DataAgeThreshold"), DataAgeThreshold, IniFile))
 			{
-				GLog->Log(ELogVerbosity::Error, TEXT("DataAgeThreshold must be specified either on the command line or in the PatchGeneration section of BuildPatchTool.ini"));
-				bSuccess = false; // We couldn't get the DataAgeThreshold value from the .ini file
+				GLog->Log(ELogVerbosity::Warning, TEXT("DataAgeThreshold not supplied, so all existing data is eligible for reuse. Note that this process is NOT compatible with any concurrently running compactify processes"));
+				DataAgeThreshold = 0.0f;
+				bPatchWithReuseAgeThreshold = false;
 			}
 		}
 	}
@@ -426,6 +432,7 @@ int32 BuildPatchToolMain( const TCHAR* CommandLine )
 		Settings.PrereqPath = PrereqPath;
 		Settings.PrereqArgs = PrereqArgs;
 		Settings.DataAgeThreshold = DataAgeThreshold;
+		Settings.bShouldHonorReuseThreshold = bPatchWithReuseAgeThreshold;
 		Settings.CustomFields = CustomFields;
 
 		// Run the build generation
diff --git a/Engine/Source/Runtime/Online/BuildPatchServices/Private/BuildPatchGeneration.cpp b/Engine/Source/Runtime/Online/BuildPatchServices/Private/BuildPatchGeneration.cpp
--- a/Engine/Source/Runtime/Online/BuildPatchServices/Private/BuildPatchGeneration.cpp
+++ b/Engine/Source/Runtime/Online/BuildPatchServices/Private/BuildPatchGeneration.cpp
@@ -996,6 +996,7 @@ bool FBuildDataGenerator::GenerateChunksManifestFromDirectory( const FBuildPatch
 	const FString& PrereqPath = Settings.PrereqPath;
 	const FString& prereqArgs = Settings.PrereqArgs;
 	const float& DataAgeThreshold = Settings.DataAgeThreshold;
+	const bool& bShouldHonorReuseThreshold = Settings.bShouldHonorReuseThreshold;
 
 	// Output to log for builder info
 	GLog->Logf( TEXT( "Running Chunks Patch Generation for: %u:%s %s" ), InAppID, *AppName, *BuildVersion );
@@ -1004,7 +1005,7 @@ bool FBuildDataGenerator::GenerateChunksManifestFromDirectory( const FBuildPatch
 	FScopeLock SingleConcurrentBuild( &SingleConcurrentBuildCS );
 
 	// Create our chunk cache
-	const FDateTime Cutoff = FDateTime::UtcNow() - FTimespan::FromDays(DataAgeThreshold);
+	const FDateTime Cutoff = bShouldHonorReuseThreshold ? FDateTime::UtcNow() - FTimespan::FromDays(DataAgeThreshold) : FDateTime::MinValue();
 	FBuildGenerationChunkCache::Init(Cutoff);
 
 	// Create a manifest
@@ -1197,6 +1198,7 @@ bool FBuildDataGenerator::GenerateFilesManifestFromDirectory( const FBuildPatchS
 	const FString& LaunchCommand = Settings.LaunchCommand;
 	const FString& IgnoreListFile = Settings.IgnoreListFile;
 	const float& DataAgeThreshold = Settings.DataAgeThreshold;
+	const bool& bShouldHonorReuseThreshold = Settings.bShouldHonorReuseThreshold;
 
 	// Output to log for builder info
 	GLog->Logf( TEXT( "Running Files Patch Generation for: %u:%s %s" ), InAppID, *AppName, *BuildVersion );
@@ -1215,7 +1217,7 @@ bool FBuildDataGenerator::GenerateFilesManifestFromDirectory( const FBuildPatchS
 	BuildManifest->DestroyData(); 
 
 	// Declare a build processor
-	const FDateTime Cutoff = FDateTime::UtcNow() - FTimespan::FromDays(DataAgeThreshold);
+	const FDateTime Cutoff = bShouldHonorReuseThreshold ? FDateTime::UtcNow() - FTimespan::FromDays(DataAgeThreshold) : FDateTime::MinValue();
 	FBuildDataFileProcessor DataProcessor(BuildManifest, RootDirectory, Cutoff);
 
 	// Set the App details
diff --git a/Engine/Source/Runtime/Online/BuildPatchServices/Public/BuildPatchSettings.h b/Engine/Source/Runtime/Online/BuildPatchServices/Public/BuildPatchSettings.h
--- a/Engine/Source/Runtime/Online/BuildPatchServices/Public/BuildPatchSettings.h
+++ b/Engine/Source/Runtime/Online/BuildPatchServices/Public/BuildPatchSettings.h
@@ -28,6 +28,8 @@ struct FBuildPatchSettings
 	FString PrereqArgs;
 	// The maximum age (in days) of existing data files which can be reused in this build
 	float DataAgeThreshold;
+	// Indicates whether data age threshold should be honored. If false, ALL data files can be reused
+	bool bShouldHonorReuseThreshold;
 	// Map of custom fields to add to the manifest
 	TMap<FString, FVariant> CustomFields;
 };
