Merging UE4-Releases/4.5 to UE4

CL#2311038 on 26/09/2014 09:59:25
------------------------------------------
Fix for BuildPatchTool.exe crash. We do not need the localization system anymore, so removing it's init.

diff --git a/Engine/Source/Programs/BuildPatchTool/Private/BuildPatchToolMain.cpp b/Engine/Source/Programs/BuildPatchTool/Private/BuildPatchToolMain.cpp
--- a/Engine/Source/Programs/BuildPatchTool/Private/BuildPatchToolMain.cpp
+++ b/Engine/Source/Programs/BuildPatchTool/Private/BuildPatchToolMain.cpp
@@ -379,10 +379,6 @@ int32 BuildPatchToolMain( const TCHAR* CommandLine )
 
 	// Initialize the file manager
 	IFileManager::Get().ProcessCommandLineOptions();
-	
-	// Populate cultures for localization (used by BuildPatch progress info)
-	BeginInitTextLocalization();
-	FInternationalization I18N = FInternationalization::Get();
 
 	// Load the BuildPatchServices Module
 	TSharedPtr<IBuildPatchServicesModule> BuildPatchServicesModule = StaticCastSharedPtr<IBuildPatchServicesModule>( FModuleManager::Get().LoadModule( TEXT( "BuildPatchServices" ) ) );
