[INTEGRATE] UE4-LauncherDev->UE4 integrate all up to cl 2475889

diff --git a/Engine/Source/Developer/FriendsAndChat/Public/FriendsAndChatStyle.h b/Engine/Source/Developer/FriendsAndChat/Public/FriendsAndChatStyle.h
--- a/Engine/Source/Developer/FriendsAndChat/Public/FriendsAndChatStyle.h
+++ b/Engine/Source/Developer/FriendsAndChat/Public/FriendsAndChatStyle.h
@@ -4,6 +4,15 @@
 
 #include "FriendsAndChatStyle.generated.h"
 
+UENUM()
+namespace EChatChannelStyle
+{
+	enum Type
+	{
+		List,
+		Conversation
+	};
+}
 
 /**
  * Interface for the services manager.
@@ -50,6 +59,14 @@ struct FRIENDSANDCHAT_API FFriendsAndChatStyle
 	FButtonStyle FriendGeneralButtonStyle;
 	FFriendsAndChatStyle& SetFriendGeneralButtonStyle(const FButtonStyle& ButtonStyle);
 
+	UPROPERTY()
+	FSlateColor ButtonInvertedForegroundColor;
+	FFriendsAndChatStyle& SetButtonInvertedForegroundColor(const FSlateColor& Value);
+
+	UPROPERTY()
+	FSlateColor ButtonForegroundColor;
+	FFriendsAndChatStyle& SetButtonForegroundColor(const FSlateColor& Value);
+
 	/** Friends List Action Button style */
 	UPROPERTY()
 	FButtonStyle FriendListActionButtonStyle;
@@ -82,6 +99,10 @@ struct FRIENDSANDCHAT_API FFriendsAndChatStyle
 	FButtonStyle FriendListItemButtonSimpleStyle;
 	FFriendsAndChatStyle& SetFriendsListItemButtonSimpleStyle(const FButtonStyle& ButtonStyle);
 
+	UPROPERTY()
+	FSlateBrush GlobalChatHeaderBrush;
+	FFriendsAndChatStyle& SetGlobalChatHeaderBrush(const FSlateBrush& Value);
+
 	UPROPERTY()
 	FButtonStyle GlobalChatButtonStyle;
 	FFriendsAndChatStyle& SetGlobalChatButtonStyle(const FButtonStyle& ButtonStyle);
@@ -96,11 +117,6 @@ struct FRIENDSANDCHAT_API FFriendsAndChatStyle
 	FButtonStyle AddFriendCloseButtonStyle;
 	FFriendsAndChatStyle& SetAddFriendCloseButtonStyle(const FButtonStyle& ButtonStyle);
 
-	/** Add Friend button style */
-	UPROPERTY()
-	FButtonStyle AddFriendButtonStyle;
-	FFriendsAndChatStyle& SetAddFriendButtonStyle(const FButtonStyle& ButtonStyle);
-
 	/** Friends List Combo Button menu background image (left) */
 	UPROPERTY()
 	FSlateBrush FriendComboBackgroundLeftBrush;
@@ -126,6 +142,11 @@ struct FRIENDSANDCHAT_API FFriendsAndChatStyle
 	FSlateBrush AddFriendButtonContentBrush;
 	FFriendsAndChatStyle& SetAddFriendButtonContentBrush(const FSlateBrush& BrushStyle);
 
+	/** Optional content for the Add Friend button (hovered) */
+	UPROPERTY()
+	FSlateBrush AddFriendButtonContentHoveredBrush;
+	FFriendsAndChatStyle& SetAddFriendButtonContentHoveredBrush(const FSlateBrush& BrushStyle);
+
 	/** Title Bar brush style */
 	UPROPERTY(EditAnywhere, Category = Appearance)
 	FSlateBrush TitleBarBrush;
@@ -186,6 +207,11 @@ struct FRIENDSANDCHAT_API FFriendsAndChatStyle
 	FSlateBrush Background;
 	FFriendsAndChatStyle& SetBackgroundBrush(const FSlateBrush& InBackground);
 
+	/** Window edging style */
+	UPROPERTY(EditAnywhere, Category = Appearance)
+	FSlateBrush WindowEdgingBrush;
+	FFriendsAndChatStyle& SetWindowEdgingBrush(const FSlateBrush& Value);
+
 	/** Friend container header */
 	UPROPERTY(EditAnywhere, Category = Appearance)
 	FSlateBrush FriendContainerHeader;
@@ -196,6 +222,11 @@ struct FRIENDSANDCHAT_API FFriendsAndChatStyle
 	FSlateBrush FriendListHeader;
 	FFriendsAndChatStyle& SetFriendListHeader(const FSlateBrush& InFriendListHeader);
 
+	/** Friend user header background */
+	UPROPERTY(EditAnywhere, Category = Appearance)
+	FSlateBrush FriendUserHeaderBackground;
+	FFriendsAndChatStyle& SetFriendUserHeaderBackground(const FSlateBrush& InFriendUserHeaderBackground);
+
 	/** Friend item selected */
 	UPROPERTY(EditAnywhere, Category = Appearance)
 	FSlateBrush FriendItemSelected;
@@ -233,8 +264,8 @@ struct FRIENDSANDCHAT_API FFriendsAndChatStyle
 
 	/** Font Style */
 	UPROPERTY(EditAnywhere, Category = Appearance)
-	FSlateFontInfo FriendsFontStyleBoldLarge;
-	FFriendsAndChatStyle& SetFontStyleBoldLarge(const FSlateFontInfo& InFontStyle);
+	FSlateFontInfo FriendsFontStyleUserLarge;
+	FFriendsAndChatStyle& SetFontStyleUserLarge(const FSlateFontInfo& InFontStyle);
 
 	/** Font Style */
 	UPROPERTY(EditAnywhere, Category=Appearance)
@@ -256,6 +287,10 @@ struct FRIENDSANDCHAT_API FFriendsAndChatStyle
 	FLinearColor DefaultFontColor;
 	FFriendsAndChatStyle& SetDefaultFontColor(const FLinearColor& InFontColor);
 
+	UPROPERTY(EditAnywhere, Category = Appearance)
+	FLinearColor DefaultDullFontColor;
+	FFriendsAndChatStyle& SetDefaultDullFontColor(const FLinearColor& InFontColor);
+
 	UPROPERTY(EditAnywhere, Category=Appearance)
 	FLinearColor DefaultChatColor;
 	FFriendsAndChatStyle& SetDefaultChatColor(const FLinearColor& InFontColor);
@@ -280,14 +315,6 @@ struct FRIENDSANDCHAT_API FFriendsAndChatStyle
 	FLinearColor ComboItemTextColorHovered;
 	FFriendsAndChatStyle& SetComboItemTextColorHovered(const FLinearColor& InColor);
 
-	UPROPERTY(EditAnywhere, Category = Appearance)
-	FLinearColor FriendListCriticalFontColor;
-	FFriendsAndChatStyle& SetFriendListCriticalFontColor(const FLinearColor& InColor);
-
-	UPROPERTY(EditAnywhere, Category = Appearance)
-	FLinearColor FriendListEmphasisFontColor;
-	FFriendsAndChatStyle& SetFriendListEmphasisFontColor(const FLinearColor& InColor);
-
 	UPROPERTY(EditAnywhere, Category = Appearance)
 	FLinearColor FriendListActionFontColor;
 	FFriendsAndChatStyle& SetFriendListActionFontColor(const FLinearColor& InColor);
@@ -316,6 +343,10 @@ struct FRIENDSANDCHAT_API FFriendsAndChatStyle
 	FTextBlockStyle ComboItemTextStyle;
 	FFriendsAndChatStyle& SetComboItemTextStyle(const FTextBlockStyle& InTextStyle);
 
+	UPROPERTY(EditAnywhere, Category = Appearance)
+	FTextBlockStyle FriendsComboTextStyle;
+	FFriendsAndChatStyle& SetFriendsComboTextStyle(const FTextBlockStyle& InTextStyle);
+
 	UPROPERTY(EditAnywhere, Category = Appearance)
 	FEditableTextBoxStyle ChatEditableTextStyle;
 	FFriendsAndChatStyle& SetChatEditableTextStyle(const FEditableTextBoxStyle& InEditableTextStyle);
@@ -332,6 +363,14 @@ struct FRIENDSANDCHAT_API FFriendsAndChatStyle
 	FVector2D ActionComboButtonSize;
 	FFriendsAndChatStyle& SetActionComboButtonSize(const FVector2D& InActionComboButtonSize);
 
+	UPROPERTY(EditAnywhere, Category = Appearance)
+	FComboButtonStyle ActionComboButtonStyle;
+	FFriendsAndChatStyle& SetActionComboButtonStyle(const FComboButtonStyle& ButtonStyle);
+
+	UPROPERTY(EditAnywhere, Category = Appearance)
+	FTextBlockStyle ActionComboButtonTextStyle;
+	FFriendsAndChatStyle& SetActionComboButtonTextStyle(const FTextBlockStyle& TextStyle);
+
 	UPROPERTY(EditAnywhere, Category = Appearance)
 	FVector2D UserPresenceImageSize;
 	FFriendsAndChatStyle& SetUserPresenceImageSize(const FVector2D& InUserPresenceImageSize);
@@ -356,6 +395,22 @@ struct FRIENDSANDCHAT_API FFriendsAndChatStyle
 	FSlateBrush ChatBackgroundBrush;
 	FFriendsAndChatStyle& SetChatBackgroundBrush(const FSlateBrush& InChatBackgroundBrush);
 
+	UPROPERTY(EditAnywhere, Category = Appearance)
+	FSlateBrush ChatBubbleBackgroundBrush;
+	FFriendsAndChatStyle& SetChatBubbleBackgroundBrush(const FSlateBrush& Value);
+
+	UPROPERTY(EditAnywhere, Category = Appearance)
+	FSlateBrush ChatBubbleLeftCallout;
+	FFriendsAndChatStyle& SetChatBubbleLeftCalloutBrush(const FSlateBrush& Value);
+
+	UPROPERTY(EditAnywhere, Category = Appearance)
+	FSlateBrush ChatBubbleRightCallout;
+	FFriendsAndChatStyle& SetChatBubbleRightCalloutBrush(const FSlateBrush& Value);
+
+	UPROPERTY(EditAnywhere, Category = Appearance)
+	FSlateBrush ChatFooterBrush;
+	FFriendsAndChatStyle& SetChatFooterBrush(const FSlateBrush& Value);
+
 	UPROPERTY(EditAnywhere, Category = Appearance)
 	FSlateBrush ChatChannelsBackgroundBrush;
 	FFriendsAndChatStyle& SetChatChannelsBackgroundBrush(const FSlateBrush& InChatChannelsBackgroundBrush);
@@ -387,6 +442,10 @@ struct FRIENDSANDCHAT_API FFriendsAndChatStyle
 	UPROPERTY(EditAnywhere, Category = Appearance)
 	bool HasUserHeader;
 	FFriendsAndChatStyle& SetHasUserHeader(bool InHasUserHeader);
+
+	UPROPERTY(EditAnywhere, Category = Appearance)
+	TEnumAsByte<EChatChannelStyle::Type> ChatChannelStyle;
+	FFriendsAndChatStyle&  SetChatWindowStyle(const EChatChannelStyle::Type Style);
 };
 
 /** Manages the style which provides resources for the rich text widget. */
diff --git a/Engine/Source/Developer/FriendsAndChat/Public/Interfaces/IFriendsAndChatManager.h b/Engine/Source/Developer/FriendsAndChat/Public/Interfaces/IFriendsAndChatManager.h
--- a/Engine/Source/Developer/FriendsAndChat/Public/Interfaces/IFriendsAndChatManager.h
+++ b/Engine/Source/Developer/FriendsAndChat/Public/Interfaces/IFriendsAndChatManager.h
@@ -7,6 +7,11 @@ namespace EChatMessageType
 {
 	enum Type : uint8;
 }
+namespace EOnlinePresenceState
+{
+	enum Type : uint8;
+}
+
 
 /**
  * Interface for the Friends and chat manager.
@@ -26,8 +31,9 @@ public:
 	 * @param InStyle The style to use to create the widgets
 	 * @param ChatType The type of chat window to create
 	 * @param FriendItem The friend if this is a whisper chat window
+	 * @param True if chat window comes to front, else opens minimized
 	 */
-	virtual void CreateChatWindow(const struct FFriendsAndChatStyle* InStyle, EChatMessageType::Type ChatType, TSharedPtr<IFriendItem> FriendItem) = 0;
+	virtual void CreateChatWindow(const struct FFriendsAndChatStyle* InStyle, EChatMessageType::Type ChatType, TSharedPtr<IFriendItem> FriendItem, bool BringToFront = false) = 0;
 
 	/**
 	 * Set the FriendsAndChatUserSettings.
@@ -50,6 +56,13 @@ public:
 	 */
 	virtual TSharedPtr< SWidget > GenerateFriendsListWidget( const struct FFriendsAndChatStyle* InStyle ) = 0;
 
+	/**
+	 * Create a status widget for the current user.
+	 * @param InStyle The style to use to create the widgets.
+	 * @return The Statuswidget.
+	 */
+	virtual TSharedPtr< SWidget > GenerateStatusWidget( const FFriendsAndChatStyle* InStyle ) = 0;
+
 	/**
 	 * Generate a chat widget.
 	 * @param InStyle The style to use to create the widgets.
@@ -77,6 +90,12 @@ public:
 	 */
 	virtual void JoinPublicChatRoom(const FString& RoomName) = 0;
 
+	/**
+	 * Show the global chat room button
+	 * @param True to show the button
+	 */
+	virtual void SetPublicChatRoomPermission(bool Allow) = 0;
+
 	/**
 	 * Delegate when the chat room has been joined
 	 */
@@ -91,6 +110,19 @@ public:
 	/** Is the chat manager logged in. */
 	virtual bool IsLoggedIn() = 0;
 
+	/** Set the user to an online state. */
+	virtual void SetOnline() = 0;
+
+	/** Set the user to an away state. */
+	virtual void SetAway() = 0;
+
+	/**
+	* Get the online status
+	*
+	* @return EOnlinePresenceState
+	*/
+	virtual EOnlinePresenceState::Type GetOnlineStatus() = 0;
+
 	/** 
 	 * Set the application view model to query and perform actions on.
 	 * @param ClientID The ID of the application
diff --git a/Engine/Source/Developer/FriendsAndChat/Public/Models/FriendsAndChatMessage.h b/Engine/Source/Developer/FriendsAndChat/Public/Models/FriendsAndChatMessage.h
--- a/Engine/Source/Developer/FriendsAndChat/Public/Models/FriendsAndChatMessage.h
+++ b/Engine/Source/Developer/FriendsAndChat/Public/Models/FriendsAndChatMessage.h
@@ -146,6 +146,16 @@ public:
 		MessageDescriptions.Add( Description );
 	}
 
+	/**
+	* Set a button style.
+	*
+	* @param ButtonStyle The button callback.
+	*/
+	void SetButtonStyle(FName ButtonStyle)
+	{
+		ButtonStyles.Add(ButtonStyle);
+	}
+
 	/**
 	 * Set the message type.
 	 *
@@ -262,6 +272,16 @@ public:
 		return MessageDescriptions;
 	}
 
+	/**
+	* Get button styles.
+	*
+	* @return The button style array.
+	*/
+	const TArray< FName >& GetButtonStyles() const
+	{
+		return ButtonStyles;
+	}
+
 	/** Set this message into a handled state. */
 	void SetHandled()
 	{
@@ -273,8 +293,11 @@ private:
 	// Holds the button callbacks
 	TArray< FOnClicked > ButtonCallbacks;
 
+	// Holds the button styles
+	TArray< FName > ButtonStyles;
+
 	// Holds the Action Descriptions
-	TArray< FText > MessageDescriptions;
+	TArray< FText > MessageDescriptions;	
 	
 	// Holds the message content
 	FString MessageConent;
diff --git a/Engine/Source/Programs/AutomationTool/MCPPublic.cs b/Engine/Source/Programs/AutomationTool/MCPPublic.cs
--- a/Engine/Source/Programs/AutomationTool/MCPPublic.cs
+++ b/Engine/Source/Programs/AutomationTool/MCPPublic.cs
@@ -331,6 +331,18 @@ namespace EpicGames.MCP.Automation
 			public List<KeyValuePair<string, float>> CustomFloatArgs;
 		}
 
+		public class DataEnumerationOptions
+		{
+			/// <summary>
+			/// Matches the corresponding BuildPatchTool command line argument.
+			/// </summary>
+			public string ManifestFile;
+			/// <summary>
+			/// Matches the corresponding BuildPatchTool command line argument.
+			/// </summary>
+			public string OutputFile;
+		}
+
         static BuildPatchToolBase Handler = null;
 
         public static BuildPatchToolBase Get()
@@ -364,6 +376,12 @@ namespace EpicGames.MCP.Automation
 		/// <param name="Opts">Parameters which will be passed to the patch tool generation process</param>
 		public abstract void Execute(PatchGenerationOptions Opts);
 
+		/// <summary>
+		/// Runs the Build Patch Tool executable to enumerate patch data files referenced by a manifest using the supplied parameters.
+		/// </summary>
+		/// <param name="Opts">Parameters which will be passed to the patch tool enumeration process</param>
+		public abstract void Execute(DataEnumerationOptions Opts);
+
     }
 
 
diff --git a/Engine/Source/Programs/BuildPatchTool/Private/BuildPatchToolMain.cpp b/Engine/Source/Programs/BuildPatchTool/Private/BuildPatchToolMain.cpp
--- a/Engine/Source/Programs/BuildPatchTool/Private/BuildPatchToolMain.cpp
+++ b/Engine/Source/Programs/BuildPatchTool/Private/BuildPatchToolMain.cpp
@@ -69,6 +69,14 @@
 	-CloudDir="E:\BuildPatchCloud" -compactify -DataAgeThreshold=14
 	-CloudDir="E:\BuildPatchCloud" -compactify -DataAgeThreshold=14 -ManifestsList="Example_V1.manifest,Example_V2.manifest"
 	-CloudDir="E:\BuildPatchCloud" -compactify -DataAgeThreshold=14 -ManifestsFile="preserve_manifests.txt"
+
+	DATA ENUMERATE MODE
+
+	Required arguments:
+	-ManifestFile=""	Specifies in quotes the file path to the manifest to enumerate from
+	-OutputFile=""		Specifies in quotes the file path to a file where the list will be saved out, \r\n separated cloud relative paths.
+	-dataenumerate		Must be specified to launch the tool in cloud seeder mode
+
 =============================================================================*/
 
 #include "RequiredProgramMainCPPInclude.h"
@@ -87,9 +95,11 @@ IMPLEMENT_APPLICATION( BuildPatchTool, "BuildPatchTool" );
 
 struct FCommandLineMatcher
 {
-	FString Command;
+	const FString Command;
 
-	FCommandLineMatcher(){}
+	FCommandLineMatcher(FString Command)
+		: Command(MoveTemp(Command))
+	{}
 
 	FORCEINLINE bool operator()(const FString& ToMatch) const
 	{
@@ -116,6 +126,18 @@ public:
 	}
 };
 
+namespace EBuildPatchToolMode
+{
+	enum Type
+	{
+		PatchGeneration,
+		Compactify,
+		DataEnumeration,
+
+		Unknown
+	};
+}
+
 
 int32 BuildPatchToolMain( const TCHAR* CommandLine )
 {
@@ -154,9 +176,10 @@ int32 BuildPatchToolMain( const TCHAR* CommandLine )
 	float DataAgeThreshold = 0.0f;
 	FString IniFile;
 	TMap<FString, FVariant> CustomFields;
+	FString ManifestFile;
+	FString OutputFile;
 
-	bool bCompactify = false;
-	bool bPatchGeneration = true;
+	EBuildPatchToolMode::Type ToolMode = EBuildPatchToolMode::Unknown;
 	bool bPreview = false;
 	bool bNoPatchDelete = false;
 	bool bPatchWithReuseAgeThreshold = true;
@@ -181,87 +204,82 @@ int32 BuildPatchToolMain( const TCHAR* CommandLine )
 		int32 ManifestsListIdx;
 		int32 ManifestsFileIdx;
 		int32 DataAgeThresholdIdx;
-
-		FCommandLineMatcher Matcher;
+		int32 ManifestFileIdx;
+		int32 OutputFileIdx;
 		bSuccess = true;
 
-		Matcher.Command = TEXT("compactify");
-		bCompactify = Switches.IndexOfByPredicate(Matcher) != INDEX_NONE;
-		bPatchGeneration = !bCompactify;
-
-		Matcher.Command = TEXT("preview");
-		bPreview = bCompactify && Switches.IndexOfByPredicate(Matcher) != INDEX_NONE;
-
-		Matcher.Command = TEXT("nopatchdelete");
-		bNoPatchDelete = bCompactify && Switches.IndexOfByPredicate(Matcher) != INDEX_NONE;
-		
-		Matcher.Command = TEXT( "BuildRoot" );
-		BuildRootIdx = Switches.IndexOfByPredicate(Matcher);
-
-		Matcher.Command = TEXT( "CloudDir" );
-		CloudDirIdx = Switches.IndexOfByPredicate(Matcher);
-
-		Matcher.Command = TEXT( "AppID" );
-		AppIDIdx = Switches.IndexOfByPredicate(Matcher);
-
-		Matcher.Command = TEXT( "AppName" );
-		AppNameIdx = Switches.IndexOfByPredicate(Matcher);
-
-		Matcher.Command = TEXT( "BuildVersion" );
-		BuildVersionIdx = Switches.IndexOfByPredicate(Matcher);
-
-		Matcher.Command = TEXT( "AppLaunch" );
-		AppLaunchIdx = Switches.IndexOfByPredicate(Matcher);
-
-		Matcher.Command = TEXT( "AppArgs" );
-		AppArgsIdx = Switches.IndexOfByPredicate(Matcher);
-
-		Matcher.Command = TEXT( "FileIgnoreList" );
-		FileIgnoreListIdx = Switches.IndexOfByPredicate(Matcher);
-
-		Matcher.Command = TEXT("FileAttributeList");
-		FileAttributeListIdx = Switches.IndexOfByPredicate(Matcher);
-
-		Matcher.Command = TEXT( "PrereqName" );
-		PrereqNameIdx = Switches.IndexOfByPredicate(Matcher);
-
-		Matcher.Command = TEXT( "PrereqPath" );
-		PrereqPathIdx = Switches.IndexOfByPredicate(Matcher);
-
-		Matcher.Command = TEXT( "PrereqArgs" );
-		PrereqArgsIdx = Switches.IndexOfByPredicate(Matcher);
-
-		Matcher.Command = TEXT("ManifestsList");
-		ManifestsListIdx = Switches.IndexOfByPredicate(Matcher);
-
-		Matcher.Command = TEXT("ManifestsFile");
-		ManifestsFileIdx = Switches.IndexOfByPredicate(Matcher);
-
-		Matcher.Command = TEXT("DataAgeThreshold");
-		DataAgeThresholdIdx = Switches.IndexOfByPredicate(Matcher);
+		if (Switches.IndexOfByPredicate(FCommandLineMatcher(TEXT("compactify"))) != INDEX_NONE)
+		{
+			ToolMode = EBuildPatchToolMode::Compactify;
+		}
+		else if (Switches.IndexOfByPredicate(FCommandLineMatcher(TEXT("dataenumerate"))) != INDEX_NONE)
+		{
+			ToolMode = EBuildPatchToolMode::DataEnumeration;
+		}
+		else
+		{
+			ToolMode = EBuildPatchToolMode::PatchGeneration;
+		}
+		bPreview = Switches.IndexOfByPredicate(FCommandLineMatcher(TEXT("preview"))) != INDEX_NONE;
+		bNoPatchDelete = Switches.IndexOfByPredicate(FCommandLineMatcher(TEXT("nopatchdelete"))) != INDEX_NONE;
+		BuildRootIdx = Switches.IndexOfByPredicate(FCommandLineMatcher(TEXT("BuildRoot")));
+		CloudDirIdx = Switches.IndexOfByPredicate(FCommandLineMatcher(TEXT("CloudDir")));
+		AppIDIdx = Switches.IndexOfByPredicate(FCommandLineMatcher(TEXT("AppID")));
+		AppNameIdx = Switches.IndexOfByPredicate(FCommandLineMatcher(TEXT("AppName")));
+		BuildVersionIdx = Switches.IndexOfByPredicate(FCommandLineMatcher(TEXT("BuildVersion")));
+		AppLaunchIdx = Switches.IndexOfByPredicate(FCommandLineMatcher(TEXT("AppLaunch")));
+		AppArgsIdx = Switches.IndexOfByPredicate(FCommandLineMatcher(TEXT("AppArgs")));
+		FileIgnoreListIdx = Switches.IndexOfByPredicate(FCommandLineMatcher(TEXT("FileIgnoreList")));
+		FileAttributeListIdx = Switches.IndexOfByPredicate(FCommandLineMatcher(TEXT("FileAttributeList")));
+		PrereqNameIdx = Switches.IndexOfByPredicate(FCommandLineMatcher(TEXT("PrereqName")));
+		PrereqPathIdx = Switches.IndexOfByPredicate(FCommandLineMatcher(TEXT("PrereqPath")));
+		PrereqArgsIdx = Switches.IndexOfByPredicate(FCommandLineMatcher(TEXT("PrereqArgs")));
+		ManifestsListIdx = Switches.IndexOfByPredicate(FCommandLineMatcher(TEXT("ManifestsList")));
+		ManifestsFileIdx = Switches.IndexOfByPredicate(FCommandLineMatcher(TEXT("ManifestsFile")));
+		DataAgeThresholdIdx = Switches.IndexOfByPredicate(FCommandLineMatcher(TEXT("DataAgeThreshold")));
+		ManifestFileIdx = Switches.IndexOfByPredicate(FCommandLineMatcher(TEXT("ManifestFile")));
+		OutputFileIdx = Switches.IndexOfByPredicate(FCommandLineMatcher(TEXT("OutputFile")));
 
 		// Check required param indexes
-		bSuccess = bSuccess && CloudDirIdx != INDEX_NONE;
-		if (bPatchGeneration)
+		switch (ToolMode)
 		{
+		case EBuildPatchToolMode::PatchGeneration:
+		bSuccess = bSuccess && CloudDirIdx != INDEX_NONE;
 			bSuccess = bSuccess && BuildRootIdx != INDEX_NONE;
 			bSuccess = bSuccess && AppIDIdx != INDEX_NONE;
 			bSuccess = bSuccess && AppNameIdx != INDEX_NONE;
 			bSuccess = bSuccess && BuildVersionIdx != INDEX_NONE;
 			bSuccess = bSuccess && AppLaunchIdx != INDEX_NONE;
 			bSuccess = bSuccess && AppArgsIdx != INDEX_NONE;
+			break;
+		case EBuildPatchToolMode::Compactify:
+			bSuccess = bSuccess && CloudDirIdx != INDEX_NONE;
+			break;
+		case EBuildPatchToolMode::DataEnumeration:
+			bSuccess = bSuccess && ManifestFileIdx != INDEX_NONE;
+			bSuccess = bSuccess && OutputFileIdx != INDEX_NONE;
+			break;
 		}
 
 		// Get required param values
-		bSuccess = bSuccess && FParse::Value( *Switches[CloudDirIdx], TEXT( "CloudDir=" ), CloudDirectory );
-		if (bPatchGeneration)
+		switch (ToolMode)
 		{
+		case EBuildPatchToolMode::PatchGeneration:
+		bSuccess = bSuccess && FParse::Value( *Switches[CloudDirIdx], TEXT( "CloudDir=" ), CloudDirectory );
 			bSuccess = bSuccess && FParse::Value(*Switches[BuildRootIdx], TEXT("BuildRoot="), RootDirectory);
 			bSuccess = bSuccess && FParse::Value(*Switches[AppIDIdx], TEXT("AppID="), AppID);
 			bSuccess = bSuccess && FParse::Value(*Switches[AppNameIdx], TEXT("AppName="), AppName);
 			bSuccess = bSuccess && FParse::Value(*Switches[BuildVersionIdx], TEXT("BuildVersion="), BuildVersion);
 			bSuccess = bSuccess && FParse::Value(*Switches[AppLaunchIdx], TEXT("AppLaunch="), LaunchExe);
 			bSuccess = bSuccess && FParse::Value(*Switches[AppArgsIdx], TEXT("AppArgs="), LaunchCommand);
+			break;
+		case EBuildPatchToolMode::Compactify:
+			bSuccess = bSuccess && FParse::Value(*Switches[CloudDirIdx], TEXT("CloudDir="), CloudDirectory);
+			break;
+		case EBuildPatchToolMode::DataEnumeration:
+			bSuccess = bSuccess && FParse::Value(*Switches[ManifestFileIdx], TEXT("ManifestFile="), ManifestFile);
+			bSuccess = bSuccess && FParse::Value(*Switches[OutputFileIdx], TEXT("OutputFile="), OutputFile);
+			break;
 		}
 
 		// Get optional param values
@@ -358,14 +376,16 @@ int32 BuildPatchToolMain( const TCHAR* CommandLine )
 			// Initialize the configuration system, we can only do this reliably if we have CloudDirectory (i.e. bSuccess is true)
 			IniFile = CloudDirectory / TEXT("BuildPatchTool.ini");
 			GConfig->InitializeConfigSystem();
-		}
 
 		if (DataAgeThresholdIdx != INDEX_NONE)
 		{
 			FParse::Value(*Switches[DataAgeThresholdIdx], TEXT("DataAgeThreshold="), DataAgeThreshold);
 		}
-		else if (bSuccess && bCompactify)
+			else
+			{
+				switch (ToolMode)
 		{
+				case EBuildPatchToolMode::Compactify:
 			// For compactification, if we don't pass in DataAgeThreshold, and it's not in BuildPatchTool.ini,
 			// then we set it to zero, to indicate that any unused chunks are valid for deletion
 			if (!GConfig->GetFloat(TEXT("Compactify"), TEXT("DataAgeThreshold"), DataAgeThreshold, IniFile))
@@ -373,9 +393,8 @@ int32 BuildPatchToolMain( const TCHAR* CommandLine )
 				GLog->Log(ELogVerbosity::Warning, TEXT("DataAgeThreshold not supplied, so all unreferenced data is eliglble for deletion. Note that this process is NOT compatible with any concurrently running patch generaiton processes"));
 				DataAgeThreshold = 0.0f;
 			}
-		}
-		else if (bSuccess && bPatchGeneration)
-		{
+					break;
+				case EBuildPatchToolMode::PatchGeneration:
 			// For patch generation, if we don't pass in DataAgeThreshold, and it's not specified in BuildPatchTool.ini,
 			// then we set bChunkWithReuseAgeThreshold to false, which indicates that *all* patch data is valid for reuse
 			if (!GConfig->GetFloat(TEXT("PatchGeneration"), TEXT("DataAgeThreshold"), DataAgeThreshold, IniFile))
@@ -383,6 +402,9 @@ int32 BuildPatchToolMain( const TCHAR* CommandLine )
 				GLog->Log(ELogVerbosity::Warning, TEXT("DataAgeThreshold not supplied, so all existing data is eligible for reuse. Note that this process is NOT compatible with any concurrently running compactify processes"));
 				DataAgeThreshold = 0.0f;
 				bPatchWithReuseAgeThreshold = false;
+			}
+					break;
+				}
 			}
 		}
 	}
@@ -397,7 +419,7 @@ int32 BuildPatchToolMain( const TCHAR* CommandLine )
 		return 1;
 	}
 
-	if (bCompactify && bPreview && bNoPatchDelete)
+	if (ToolMode == EBuildPatchToolMode::Compactify && bPreview && bNoPatchDelete)
 	{
 		GLog->Log(ELogVerbosity::Error, TEXT("Only one of -preview and -nopatchdelete can be specified"));
 		return 5;
@@ -426,7 +448,10 @@ int32 BuildPatchToolMain( const TCHAR* CommandLine )
 	// Setup the module
 	BuildPatchServicesModule->SetCloudDirectory( CloudDirectory + TEXT( "/" ) );
 
-	if (bCompactify)
+	// Run the mode!
+	switch (ToolMode)
+	{
+		case EBuildPatchToolMode::Compactify:
 	{
 		// Split out our manifests to keep arg (if any) into an array of manifest filenames
 		TArray<FString> ManifestsArr;
@@ -466,7 +491,8 @@ int32 BuildPatchToolMain( const TCHAR* CommandLine )
 		// Run the compactify routine
 		bSuccess = BuildPatchServicesModule->CompactifyCloudDirectory(ManifestsArr, DataAgeThreshold, CompactifyMode);
 	}
-	else if (bPatchGeneration)
+		break;
+		case EBuildPatchToolMode::PatchGeneration:
 	{
 		FBuildPatchSettings Settings;
 		Settings.RootDirectory = RootDirectory + TEXT("/");
@@ -494,11 +520,19 @@ int32 BuildPatchToolMain( const TCHAR* CommandLine )
 			bSuccess = BuildPatchServicesModule->GenerateChunksManifestFromDirectory(Settings);
 		}
 	}
-	else
+		break;
+		case EBuildPatchToolMode::DataEnumeration:
+		{
+			// Run the data enumeration routine
+			bSuccess = BuildPatchServicesModule->EnumerateManifestData(MoveTemp(ManifestFile), MoveTemp(OutputFile));
+		}
+		break;
+		default:
 	{
 		GLog->Log(ELogVerbosity::Error, TEXT("Unknown tool mode"));
 		BuildPatchServicesModule.Reset();
 		FCoreDelegates::OnExit.Broadcast();
+		}
 		return 3;
 	}
 
diff --git a/Engine/Source/Runtime/Analytics/AnalyticsET/Private/AnalyticsET.cpp b/Engine/Source/Runtime/Analytics/AnalyticsET/Private/AnalyticsET.cpp
--- a/Engine/Source/Runtime/Analytics/AnalyticsET/Private/AnalyticsET.cpp
+++ b/Engine/Source/Runtime/Analytics/AnalyticsET/Private/AnalyticsET.cpp
@@ -157,13 +157,15 @@ FAnalyticsProviderET::FAnalyticsProviderET(const FAnalyticsET::Config& ConfigVal
 	, bShouldCacheEvents(!FParse::Param(FCommandLine::Get(), TEXT("ANALYTICSDISABLECACHING")))
 	, FlushEventsCountdown(MaxCachedElapsedTime)
 {
-	UE_LOG(LogAnalytics, Verbose, TEXT("Initializing ET Analytics provider"));
-
-	APIKey = ConfigValues.APIKeyET;
-	if (APIKey.IsEmpty())
+	// if we are not caching events, we are operating in debug mode. Turn on super-verbose analytics logging
+	if (!bShouldCacheEvents)
 	{
-		UE_LOG(LogAnalytics, Warning, TEXT("AnalyticsET missing APIKey. No uploads will be processed."));
+		LogAnalytics.SetVerbosity(ELogVerbosity::VeryVerbose);
 	}
+
+	APIKey = ConfigValues.APIKeyET;
+	UE_LOG(LogAnalytics, Verbose, TEXT("[%s] Initializing ET Analytics provider"), *APIKey);
+
 	// allow the APIServer value to be empty and use defaults.
 	APIServer = ConfigValues.APIServerET.IsEmpty() 
 		? FAnalyticsET::Config::GetDefaultAPIServer()
@@ -185,10 +187,10 @@ FAnalyticsProviderET::FAnalyticsProviderET(const FAnalyticsET::Config& ConfigVal
 		? GEngineVersion.ToString() 
 		: ConfigAppVersion.Replace(TEXT("%VERSION%"), *GEngineVersion.ToString(), ESearchCase::CaseSensitive);
 
-	UE_LOG(LogAnalytics, Log, TEXT("ET APIKey = %s. APIServer = %s. AppVersion = %s"), *APIKey, *APIServer, *AppVersion);
+	UE_LOG(LogAnalytics, Log, TEXT("[%s] APIServer = %s. AppVersion = %s"), *APIKey, *APIServer, *AppVersion);
 	if (UseDataRouter)
 	{
-		UE_LOG(LogAnalytics, Log, TEXT("ET APIKey = %s. DataRouterUploadURL = %s. AppVersion = %s"), *APIKey, *DataRouterUploadURL, *AppVersion);
+		UE_LOG(LogAnalytics, Log, TEXT("[%s] DataRouterUploadURL = %s. AppVersion = %s"), *APIKey, *DataRouterUploadURL, *AppVersion);
 	}
 
 	// cache the build type string
@@ -234,7 +236,7 @@ bool FAnalyticsProviderET::Tick(float DeltaSeconds)
 
 FAnalyticsProviderET::~FAnalyticsProviderET()
 {
-	UE_LOG(LogAnalytics, Verbose, TEXT("Destroying ET Analytics provider"));
+	UE_LOG(LogAnalytics, Verbose, TEXT("[%s] Destroying ET Analytics provider"), *APIKey);
 	EndSession();
 }
 
@@ -244,7 +246,7 @@ FAnalyticsProviderET::~FAnalyticsProviderET()
  */
 bool FAnalyticsProviderET::StartSession(const TArray<FAnalyticsEventAttribute>& Attributes)
 {
-	UE_LOG(LogAnalytics, Log, TEXT("AnalyticsET::StartSession [%s]"),*APIKey);
+	UE_LOG(LogAnalytics, Log, TEXT("[%s] AnalyticsET::StartSession"),*APIKey);
 
 	// end/flush previous session before staring new one
 	if (bSessionInProgress)
@@ -332,34 +334,23 @@ void FAnalyticsProviderET::FlushEvents()
 		JsonWriter->WriteObjectEnd();
 		JsonWriter->Close();
 
-		// when we are not caching events, assume we are debugging and always log analytics payloads.
-		if (bShouldCacheEvents)
-		{
-			UE_LOG(LogAnalytics, Verbose, TEXT("AnalyticsET Payload:%s"), *Payload);
-		}
-		else
-		{
-			UE_LOG(LogAnalytics, Display, TEXT("AnalyticsET Payload:%s"), *Payload);
-		}
-
+		FString URLPath = FString::Printf(TEXT("CollectData.1?SessionID=%s&AppID=%s&AppVersion=%s&UserID=%s"),
+			*FGenericPlatformHttp::UrlEncode(SessionID),
+			*FGenericPlatformHttp::UrlEncode(APIKey),
+			*FGenericPlatformHttp::UrlEncode(AppVersion),
+			*FGenericPlatformHttp::UrlEncode(UserID));
+		
 		// Create/send Http request for an event
 		TSharedRef<IHttpRequest> HttpRequest = FHttpModule::Get().CreateRequest();
 		HttpRequest->SetHeader(TEXT("Content-Type"), TEXT("application/json; charset=utf-8"));
 
-		HttpRequest->SetURL(
-			FString::Printf(TEXT("%sCollectData.1?SessionID=%s&AppID=%s&AppVersion=%s&UserID=%s&IsEditor=%s"),
-			*APIServer, 
-			*FGenericPlatformHttp::UrlEncode(SessionID),
-			*FGenericPlatformHttp::UrlEncode(APIKey), 
-			*FGenericPlatformHttp::UrlEncode(AppVersion),
-			*FGenericPlatformHttp::UrlEncode(UserID),
-			*FGenericPlatformHttp::UrlEncode(FString::FromInt(GIsEditor))
-		));
-
+		HttpRequest->SetURL(APIServer + URLPath);
 		HttpRequest->SetVerb(TEXT("POST"));
 		HttpRequest->SetContentAsString(Payload);
 		HttpRequest->OnProcessRequestComplete().BindRaw(this, &FAnalyticsProviderET::EventRequestComplete);
 
+		UE_LOG(LogAnalytics, Display, TEXT("[%s] AnalyticsET URL:%s. Payload:%s"), *APIKey, *URLPath, *Payload);
+
  		HttpRequest->ProcessRequest();
 
 		if (UseDataRouter)
@@ -399,12 +390,12 @@ void FAnalyticsProviderET::SetUserID(const FString& InUserID)
 	// command-line specified user ID overrides all attempts to reset it.
 	if (!FParse::Value(FCommandLine::Get(), TEXT("ANALYTICSUSERID="), UserID, false))
 	{
-		UE_LOG(LogAnalytics, Log, TEXT("SetUserId %s"), *InUserID);
+		UE_LOG(LogAnalytics, Log, TEXT("[%s] SetUserId %s"), *APIKey, *InUserID);
 		UserID = InUserID;
 	}
 	else if (UserID != InUserID)
 	{
-		UE_LOG(LogAnalytics, Log, TEXT("Overriding SetUserId %s with cmdline UserId of %s."), *InUserID, *UserID);
+		UE_LOG(LogAnalytics, Log, TEXT("[%s] Overriding SetUserId %s with cmdline UserId of %s."), *APIKey, *InUserID, *UserID);
 	}
 }
 
@@ -423,7 +414,7 @@ bool FAnalyticsProviderET::SetSessionID(const FString& InSessionID)
 	if (bSessionInProgress)
 	{
 		SessionID = InSessionID;
-		UE_LOG(LogAnalytics, Log, TEXT("AnalyticsET: Forcing SessionID to %s."), *SessionID);
+		UE_LOG(LogAnalytics, Log, TEXT("[%s] Forcing SessionID to %s."), *APIKey, *SessionID);
 		return true;
 	}
 	return false;
@@ -448,11 +439,11 @@ void FAnalyticsProviderET::EventRequestComplete(FHttpRequestPtr HttpRequest, FHt
 {
 	if (bSucceeded && HttpResponse.IsValid())
 	{
-		UE_LOG(LogAnalytics, VeryVerbose, TEXT("ET response for [%s]. Code: %d. Payload: %s"), *HttpRequest->GetURL(), HttpResponse->GetResponseCode(), *HttpResponse->GetContentAsString());
+		UE_LOG(LogAnalytics, VeryVerbose, TEXT("[%s] ET response for [%s]. Code: %d. Payload: %s"), *APIKey, *HttpRequest->GetURL(), HttpResponse->GetResponseCode(), *HttpResponse->GetContentAsString());
 	}
 	else
 	{
-		UE_LOG(LogAnalytics, VeryVerbose, TEXT("ET response for [%s]. No response"), *HttpRequest->GetURL());
+		UE_LOG(LogAnalytics, VeryVerbose, TEXT("[%s] ET response for [%s]. No response"), *APIKey, *HttpRequest->GetURL());
 	}
 }
 
@@ -460,10 +451,10 @@ void FAnalyticsProviderET::EventRequestCompleteDataRouter(FHttpRequestPtr HttpRe
 {
 	if (bSucceeded && HttpResponse.IsValid())
 	{
-		UE_LOG(LogAnalytics, VeryVerbose, TEXT("ET (DataRouter) response for [%s]. Code: %d. Payload: %s"), *HttpRequest->GetURL(), HttpResponse->GetResponseCode(), *HttpResponse->GetContentAsString());
+		UE_LOG(LogAnalytics, VeryVerbose, TEXT("[%s] ET (DataRouter) response for [%s]. Code: %d. Payload: %s"), *APIKey, *HttpRequest->GetURL(), HttpResponse->GetResponseCode(), *HttpResponse->GetContentAsString());
 	}
 	else
 	{
-		UE_LOG(LogAnalytics, VeryVerbose, TEXT("ET (DataRouter) response for [%s]. No response"), *HttpRequest->GetURL());
+		UE_LOG(LogAnalytics, VeryVerbose, TEXT("[%s] ET (DataRouter) response for [%s]. No response"), *APIKey, *HttpRequest->GetURL());
 	}
 }
diff --git a/Engine/Source/Runtime/Analytics/AnalyticsET/Private/HttpServiceTracker.cpp b/Engine/Source/Runtime/Analytics/AnalyticsET/Private/HttpServiceTracker.cpp
--- a/Engine/Source/Runtime/Analytics/AnalyticsET/Private/HttpServiceTracker.cpp
+++ b/Engine/Source/Runtime/Analytics/AnalyticsET/Private/HttpServiceTracker.cpp
@@ -80,34 +80,37 @@ bool FHttpServiceTracker::EndpointMetrics::IsSuccessfulResponse(int32 ResponseCo
 
 void FHttpServiceTracker::EndpointMetrics::TrackRequest(const FHttpRequestPtr& HttpRequest)
 {
-	// keep a histogram of response codes
-	const FHttpResponsePtr HttpResponse = HttpRequest->GetResponse();
-
-	const int32 ResponseCode = HttpResponse.IsValid() ? HttpResponse->GetResponseCode() : 0;
-	// track all responses in a histogram
-	ResponseCodes.FindOrAdd(ResponseCode)++;
-	const float ElapsedTime = HttpRequest->GetElapsedTime();
-	const int64 DownloadBytes = HttpRequest->GetResponse()->GetContent().Num();
-	// track successes/fails separately
-	if (IsSuccessfulResponse(ResponseCode))
-	{
-		++SuccessCount;
-		// sum elapsed time for average calc
-		ElapsedTimeSuccessTotal += ElapsedTime;
-		ElapsedTimeSuccessMax = FMath::Max(ElapsedTimeSuccessMax, ElapsedTime);
-		ElapsedTimeSuccessMin = FMath::Min(ElapsedTimeSuccessMin, ElapsedTime);
-		// sum download rate for average calc
-		DownloadBytesSuccessTotal += DownloadBytes;
-	}
-	else
+	if(HttpRequest.IsValid())
 	{
-		++FailCount;
-		// sum elapsed time for average calc
-		ElapsedTimeFailTotal += ElapsedTime;
-		ElapsedTimeFailMax = FMath::Max(ElapsedTimeFailMax, ElapsedTime);
-		ElapsedTimeFailMin = FMath::Min(ElapsedTimeFailMin, ElapsedTime);
-		// sum download rate for average calc
-		DownloadBytesFailTotal += DownloadBytes;
+		// keep a histogram of response codes
+		const FHttpResponsePtr HttpResponse = HttpRequest->GetResponse();
+
+		const int32 ResponseCode = HttpResponse.IsValid() ? HttpResponse->GetResponseCode() : 0;
+		// track all responses in a histogram
+		ResponseCodes.FindOrAdd(ResponseCode)++;
+		const float ElapsedTime = HttpRequest->GetElapsedTime();
+		const int64 DownloadBytes = HttpResponse.IsValid() ? HttpResponse->GetContent().Num() : 0;
+		// track successes/fails separately
+		if (IsSuccessfulResponse(ResponseCode))
+		{
+			++SuccessCount;
+			// sum elapsed time for average calc
+			ElapsedTimeSuccessTotal += ElapsedTime;
+			ElapsedTimeSuccessMax = FMath::Max(ElapsedTimeSuccessMax, ElapsedTime);
+			ElapsedTimeSuccessMin = FMath::Min(ElapsedTimeSuccessMin, ElapsedTime);
+			// sum download rate for average calc
+			DownloadBytesSuccessTotal += DownloadBytes;
+		}
+		else
+		{
+			++FailCount;
+			// sum elapsed time for average calc
+			ElapsedTimeFailTotal += ElapsedTime;
+			ElapsedTimeFailMax = FMath::Max(ElapsedTimeFailMax, ElapsedTime);
+			ElapsedTimeFailMin = FMath::Min(ElapsedTimeFailMin, ElapsedTime);
+			// sum download rate for average calc
+			DownloadBytesFailTotal += DownloadBytes;
+		}
 	}
 }
 
diff --git a/Engine/Source/Runtime/Core/Private/Mac/MacPlatformMisc.cpp b/Engine/Source/Runtime/Core/Private/Mac/MacPlatformMisc.cpp
--- a/Engine/Source/Runtime/Core/Private/Mac/MacPlatformMisc.cpp
+++ b/Engine/Source/Runtime/Core/Private/Mac/MacPlatformMisc.cpp
@@ -690,6 +690,11 @@ void FMacPlatformMisc::RequestExit( bool Force )
 	}
 }
 
+void FMacPlatformMisc::RequestMinimize()
+{
+	[NSApp hide : nil];
+}
+
 const TCHAR* FMacPlatformMisc::GetSystemErrorMessage(TCHAR* OutBuffer, int32 BufferCount, int32 Error)
 {
 	// There's no Mac equivalent for GetLastError()
@@ -975,6 +980,7 @@ int32 FMacPlatformMisc::NumberOfCoresIncludingHyperthreads()
 
 void FMacPlatformMisc::NormalizePath(FString& InPath)
 {
+	SCOPED_AUTORELEASE_POOL;
 	if (InPath.Len() > 1)
 	{
 		const bool bAppendSlash = InPath[InPath.Len() - 1] == '/'; // NSString will remove the trailing slash, if present, so we need to restore it after conversion
diff --git a/Engine/Source/Runtime/Core/Private/Windows/WindowsPlatformMisc.cpp b/Engine/Source/Runtime/Core/Private/Windows/WindowsPlatformMisc.cpp
--- a/Engine/Source/Runtime/Core/Private/Windows/WindowsPlatformMisc.cpp
+++ b/Engine/Source/Runtime/Core/Private/Windows/WindowsPlatformMisc.cpp
@@ -978,6 +978,11 @@ void FWindowsPlatformMisc::RequestExit( bool Force )
 	}
 }
 
+void FWindowsPlatformMisc::RequestMinimize()
+{
+	::ShowWindow(::GetActiveWindow(), SW_MINIMIZE);
+}
+
 const TCHAR* FWindowsPlatformMisc::GetSystemErrorMessage(TCHAR* OutBuffer, int32 BufferCount, int32 Error)
 {
 	check(OutBuffer && BufferCount);
diff --git a/Engine/Source/Runtime/Core/Private/Windows/WindowsPlatformProcess.cpp b/Engine/Source/Runtime/Core/Private/Windows/WindowsPlatformProcess.cpp
--- a/Engine/Source/Runtime/Core/Private/Windows/WindowsPlatformProcess.cpp
+++ b/Engine/Source/Runtime/Core/Private/Windows/WindowsPlatformProcess.cpp
@@ -284,47 +284,47 @@ FProcHandle FWindowsPlatformProcess::CreateProc( const TCHAR* URL, const TCHAR*
 	// initialize process attributes
 	SECURITY_ATTRIBUTES Attr;
 	{
-		Attr.nLength = sizeof(SECURITY_ATTRIBUTES);
-		Attr.lpSecurityDescriptor = NULL;
-		Attr.bInheritHandle = true;
+	Attr.nLength = sizeof(SECURITY_ATTRIBUTES);
+	Attr.lpSecurityDescriptor = NULL;
+	Attr.bInheritHandle = true;
 	}
 
 	// initialize process creation flags
 	uint32 CreateFlags = NORMAL_PRIORITY_CLASS;
 	{
-		if (PriorityModifier < 0)
-		{
+	if (PriorityModifier < 0)
+	{
 			CreateFlags = (PriorityModifier == -1) ? BELOW_NORMAL_PRIORITY_CLASS : IDLE_PRIORITY_CLASS;
-		}
-		else if (PriorityModifier > 0)
-		{
+	}
+	else if (PriorityModifier > 0)
+	{
 			CreateFlags = (PriorityModifier == 1) ? ABOVE_NORMAL_PRIORITY_CLASS : HIGH_PRIORITY_CLASS;
-		}
+	}
 
-		if (bLaunchDetached)
-		{
-			CreateFlags |= DETACHED_PROCESS;
-		}
+	if (bLaunchDetached)
+	{
+		CreateFlags |= DETACHED_PROCESS;
+	}
 	}
 
 	// initialize window flags
 	uint32 dwFlags = 0;
 	uint16 ShowWindowFlags = SW_HIDE;
 	{
-		if (bLaunchReallyHidden)
-		{
-			dwFlags = STARTF_USESHOWWINDOW;
-		}
-		else if (bLaunchHidden)
-		{
-			dwFlags = STARTF_USESHOWWINDOW;
-			ShowWindowFlags = SW_SHOWMINNOACTIVE;
-		}
+	if (bLaunchReallyHidden)
+	{
+		dwFlags = STARTF_USESHOWWINDOW;
+	}
+	else if (bLaunchHidden)
+	{
+		dwFlags = STARTF_USESHOWWINDOW;
+		ShowWindowFlags = SW_SHOWMINNOACTIVE;
+	}
 
 		if (PipeWrite != nullptr)
-		{
-			dwFlags |= STARTF_USESTDHANDLES;
-		}
+	{
+		dwFlags |= STARTF_USESTDHANDLES;
+	}
 	}
 
 	// initialize startup info
@@ -350,6 +350,7 @@ FProcHandle FWindowsPlatformProcess::CreateProc( const TCHAR* URL, const TCHAR*
 
 	if (!CreateProcess(NULL, CommandLine.GetCharArray().GetData(), &Attr, &Attr, true, CreateFlags, NULL, OptionalWorkingDirectory, &StartupInfo, &ProcInfo))
 	{
+		UE_LOG(LogWindows, Warning, TEXT("CreateProc failed (%u) %s %s"), ::GetLastError(), URL, Parms);
 		if (OutProcessID != nullptr)
 		{
 			*OutProcessID = 0;
@@ -363,7 +364,7 @@ FProcHandle FWindowsPlatformProcess::CreateProc( const TCHAR* URL, const TCHAR*
 		*OutProcessID = ProcInfo.dwProcessId;
 	}
 
-	::CloseHandle(ProcInfo.hThread);
+	::CloseHandle( ProcInfo.hThread );
 
 	return FProcHandle(ProcInfo.hProcess);
 }
@@ -1112,7 +1113,7 @@ bool FWindowsPlatformProcess::ReadPipeToArray(void* ReadPipe, TArray<uint8> & Ou
 
 #include "AllowWindowsPlatformTypes.h"
 
-FWindowsPlatformProcess::FWindowsSemaphore::FWindowsSemaphore(const FString& InName, HANDLE InSemaphore)
+FWindowsPlatformProcess::FWindowsSemaphore::FWindowsSemaphore(const FString & InName, HANDLE InSemaphore)
 	:	FSemaphore(InName)
 	,	Semaphore(InSemaphore)
 {
@@ -1166,7 +1167,7 @@ void FWindowsPlatformProcess::FWindowsSemaphore::Unlock()
 	}
 }
 
-FWindowsPlatformProcess::FSemaphore* FWindowsPlatformProcess::NewInterprocessSynchObject(const FString& Name, bool bCreate, uint32 MaxLocks)
+FWindowsPlatformProcess::FSemaphore * FWindowsPlatformProcess::NewInterprocessSynchObject(const FString & Name, bool bCreate, uint32 MaxLocks)
 {
 	HANDLE Semaphore = NULL;
 	
diff --git a/Engine/Source/Runtime/Core/Public/Async/SlowResult.h b/Engine/Source/Runtime/Core/Public/Async/SlowResult.h
--- a/Engine/Source/Runtime/Core/Public/Async/SlowResult.h
+++ b/Engine/Source/Runtime/Core/Public/Async/SlowResult.h
@@ -9,7 +9,7 @@
 /**
  * Template for slow asynchronous results.
  */
-template<typename ResultType)
+template<typename ResultType>
 struct TSlowResult
 {
 	/** Holds the asynchronous return value. */
@@ -18,13 +18,15 @@ struct TSlowResult
 	/** Pointer to the asynchronous task that is computing the result, or nullptr if the result was returned synchronously. */
 	TSharedPtr<ISlowTask> Task;
 
+	/** Default constructor (no initialization). */
+	TSlowResult() { }
 
 	/**
-	* Creates and initializes a new instance.
-	*
-	* @param InFuture The future object to set.
-	* @param InTask The slow task to set.
-	*/
+	 * Creates and initializes a new instance.
+	 *
+	 * @param InFuture The future object to set.
+	 * @param InTask The slow task to set.
+	 */
 	TSlowResult(TFuture<ResultType>&& InFuture, const TSharedRef<ISlowTask>& InTask)
 		: Future(InFuture)
 		, Task(InTask)
diff --git a/Engine/Source/Runtime/Core/Public/Mac/MacPlatformMisc.h b/Engine/Source/Runtime/Core/Public/Mac/MacPlatformMisc.h
--- a/Engine/Source/Runtime/Core/Public/Mac/MacPlatformMisc.h
+++ b/Engine/Source/Runtime/Core/Public/Mac/MacPlatformMisc.h
@@ -86,6 +86,7 @@ struct CORE_API FMacPlatformMisc : public FGenericPlatformMisc
 	static uint32 GetKeyMap( uint16* KeyCodes, FString* KeyNames, uint32 MaxMappings );
 	static uint32 GetCharKeyMap(uint16* KeyCodes, FString* KeyNames, uint32 MaxMappings);
 	static void RequestExit(bool Force);
+	static void RequestMinimize();
 	static const TCHAR* GetSystemErrorMessage(TCHAR* OutBuffer, int32 BufferCount, int32 Error);
 	static void ClipboardCopy(const TCHAR* Str);
 	static void ClipboardPaste(class FString& Dest);
diff --git a/Engine/Source/Runtime/Core/Public/Windows/WindowsPlatformMisc.h b/Engine/Source/Runtime/Core/Public/Windows/WindowsPlatformMisc.h
--- a/Engine/Source/Runtime/Core/Public/Windows/WindowsPlatformMisc.h
+++ b/Engine/Source/Runtime/Core/Public/Windows/WindowsPlatformMisc.h
@@ -82,6 +82,7 @@ struct CORE_API FWindowsPlatformMisc
 	static void SetUTF8Output();
 	static void LocalPrint(const TCHAR *Message);
 	static void RequestExit(bool Force);
+	static void RequestMinimize();
 	static const TCHAR* GetSystemErrorMessage(TCHAR* OutBuffer, int32 BufferCount, int32 Error);
 	static void ClipboardCopy(const TCHAR* Str);
 	static void ClipboardPaste(class FString& Dest);
diff --git a/Engine/Source/Runtime/Online/Android/OnlineSubsystemGooglePlay/Private/OnlineSubsystemGooglePlay.cpp b/Engine/Source/Runtime/Online/Android/OnlineSubsystemGooglePlay/Private/OnlineSubsystemGooglePlay.cpp
--- a/Engine/Source/Runtime/Online/Android/OnlineSubsystemGooglePlay/Private/OnlineSubsystemGooglePlay.cpp
+++ b/Engine/Source/Runtime/Online/Android/OnlineSubsystemGooglePlay/Private/OnlineSubsystemGooglePlay.cpp
@@ -67,6 +67,11 @@ IOnlineUserCloudPtr FOnlineSubsystemGooglePlay::GetUserCloudInterface() const
 	return nullptr;
 }
 
+IOnlineUserCloudPtr FOnlineSubsystemGooglePlay::GetUserCloudInterface(const FString& Key) const
+{
+	return nullptr;
+}
+
 IOnlineLeaderboardsPtr FOnlineSubsystemGooglePlay::GetLeaderboardsInterface() const
 {
 	return LeaderboardsInterface;
diff --git a/Engine/Source/Runtime/Online/Android/OnlineSubsystemGooglePlay/Public/OnlineSubsystemGooglePlay.h b/Engine/Source/Runtime/Online/Android/OnlineSubsystemGooglePlay/Public/OnlineSubsystemGooglePlay.h
--- a/Engine/Source/Runtime/Online/Android/OnlineSubsystemGooglePlay/Public/OnlineSubsystemGooglePlay.h
+++ b/Engine/Source/Runtime/Online/Android/OnlineSubsystemGooglePlay/Public/OnlineSubsystemGooglePlay.h
@@ -31,6 +31,7 @@ public:
 	virtual IOnlineGroupsPtr GetGroupsInterface() const override;
 	virtual IOnlineSharedCloudPtr GetSharedCloudInterface() const override;
 	virtual IOnlineUserCloudPtr GetUserCloudInterface() const override;
+	virtual IOnlineUserCloudPtr GetUserCloudInterface(const FString& Key) const override;
 	virtual IOnlineLeaderboardsPtr GetLeaderboardsInterface() const override;
 	virtual IOnlineVoicePtr GetVoiceInterface() const  override;
 	virtual IOnlineExternalUIPtr GetExternalUIInterface() const override;
diff --git a/Engine/Source/Runtime/Online/BuildPatchServices/Private/BuildPatchDataEnumeration.cpp b/Engine/Source/Runtime/Online/BuildPatchServices/Private/BuildPatchDataEnumeration.cpp
new file mode 100644
--- /dev/null
+++ b/Engine/Source/Runtime/Online/BuildPatchServices/Private/BuildPatchDataEnumeration.cpp
@@ -0,0 +1,49 @@
+// Copyright 1998-2015 Epic Games, Inc. All Rights Reserved.
+
+#include "BuildPatchServicesPrivatePCH.h"
+
+#if WITH_BUILDPATCHGENERATION
+
+bool FBuildDataEnumeration::EnumerateManifestData(FString ManifestFilePath, FString OutputFile)
+{
+	bool bSuccess = false;
+	TArray<FGuid> DataList;
+	FBuildPatchAppManifestRef AppManifest = MakeShareable(new FBuildPatchAppManifest());
+	GLog->Logf(ELogVerbosity::Log, TEXT("Load manifest %s"), *ManifestFilePath);
+	const double StartLoadManifest = FPlatformTime::Seconds();
+	if (AppManifest->LoadFromFile(ManifestFilePath))
+	{
+		const double LoadManifestTime = FPlatformTime::Seconds() - StartLoadManifest;
+		GLog->Logf(ELogVerbosity::Log, TEXT("Loaded manifest in %.1f seconds"), LoadManifestTime);
+
+		// Enumerate
+		AppManifest->GetDataList(DataList);
+
+		// Output text file
+		GLog->Log(ELogVerbosity::Log, TEXT("Data file list:-"));
+		FString FullList;
+		for (FGuid& DataGuid : DataList)
+		{
+			FString Path = FBuildPatchUtils::GetDataFilename(AppManifest, FString(), DataGuid);
+			FullList += Path;
+			FullList += TEXT("\r\n");
+			GLog->Log(ELogVerbosity::Log, Path);
+		}
+		if (FFileHelper::SaveStringToFile(FullList, *OutputFile))
+		{
+			GLog->Logf(ELogVerbosity::Log, TEXT("Saved out to %s"), *OutputFile);
+			bSuccess = true;
+		}
+		else
+		{
+			GLog->Logf(ELogVerbosity::Error, TEXT("Failed to save output %s"), *OutputFile);
+		}
+	}
+	else
+	{
+		GLog->Logf(ELogVerbosity::Error, TEXT("Failed to load manifest %s"), *ManifestFilePath);
+	}
+	return bSuccess;
+}
+
+#endif //WITH_BUILDPATCHGENERATION
\ No newline at end of file
diff --git a/Engine/Source/Runtime/Online/BuildPatchServices/Private/BuildPatchDataEnumeration.h b/Engine/Source/Runtime/Online/BuildPatchServices/Private/BuildPatchDataEnumeration.h
new file mode 100644
--- /dev/null
+++ b/Engine/Source/Runtime/Online/BuildPatchServices/Private/BuildPatchDataEnumeration.h
@@ -0,0 +1,13 @@
+// Copyright 1998-2015 Epic Games, Inc. All Rights Reserved.
+
+#pragma once
+
+#if WITH_BUILDPATCHGENERATION
+
+class FBuildDataEnumeration
+{
+public:
+	static bool EnumerateManifestData(FString ManifestFilePath, FString OutputFile);
+};
+
+#endif //WITH_BUILDPATCHGENERATION
diff --git a/Engine/Source/Runtime/Online/BuildPatchServices/Private/BuildPatchDownloader.cpp b/Engine/Source/Runtime/Online/BuildPatchServices/Private/BuildPatchDownloader.cpp
--- a/Engine/Source/Runtime/Online/BuildPatchServices/Private/BuildPatchDownloader.cpp
+++ b/Engine/Source/Runtime/Online/BuildPatchServices/Private/BuildPatchDownloader.cpp
@@ -230,7 +230,7 @@ uint32 FBuildPatchDownloader::Run()
 					InFlightDownloads.Remove( InFlightKey );
 				}
 			}
-			else if (MeanChunkTime.IsReliable() && InFlightJob.RetryCount.GetValue() == 0)
+			else if (bIsChunkData && MeanChunkTime.IsReliable() && InFlightJob.RetryCount.GetValue() == 0)
 			{
 				// If still on first try, cancel any chunk taking longer than the mean time plus 4x standard deviation. In statistical terms, that's 1 in 15,787 chance of being
 				// a download time that appears in the normal distribution. So we are guessing that this chunk would be an abnormally delayed one.
@@ -464,7 +464,7 @@ bool FBuildPatchDownloader::ShouldBeRunning()
 	return bWaitingForJobs || GetNumChunksLeft() > 0;
 }
 
-void FBuildPatchDownloader::HttpRequestProgress( FHttpRequestPtr Request, int32 BytesSoFar )
+void FBuildPatchDownloader::HttpRequestProgress( FHttpRequestPtr Request, int32 BytesSent, int32 BytesReceived )
 {
 #if EXTRA_DOWNLOAD_LOGGING
 	GWarn->Logf(TEXT("BuildPatchDownloader: Request %p Bytes: %d"), Request.Get(), BytesSoFar);
@@ -473,7 +473,7 @@ void FBuildPatchDownloader::HttpRequestProgress( FHttpRequestPtr Request, int32
 	{
 		FGuid DataGUID;
 		FBuildPatchUtils::GetGUIDFromFilename( Request->GetURL(), DataGUID );
-		OnDownloadProgress( DataGUID, BytesSoFar );
+		OnDownloadProgress( DataGUID, BytesReceived );
 	}
 }
 
diff --git a/Engine/Source/Runtime/Online/BuildPatchServices/Private/BuildPatchDownloader.h b/Engine/Source/Runtime/Online/BuildPatchServices/Private/BuildPatchDownloader.h
--- a/Engine/Source/Runtime/Online/BuildPatchServices/Private/BuildPatchDownloader.h
+++ b/Engine/Source/Runtime/Online/BuildPatchServices/Private/BuildPatchDownloader.h
@@ -268,10 +268,11 @@ private:
 
 	/**
 	 * Receives the HTTP request bytes downloaded when downloading a chunk.
-	 * @param Request		The request that was made
-	 * @param BytesSoFar	Whether the request successfully completed
+	 * @param Request               The request that was made
+	 * @param BytesSentSoFar        The total number of bytes sent so far
+	 * @param BytesReceivedSoFar    The total number of bytes downloaded so far
 	 */
-	void HttpRequestProgress( FHttpRequestPtr Request, int32 BytesSoFar );
+	void HttpRequestProgress( FHttpRequestPtr Request, int32 BytesSent, int32 BytesReceived );
 
 	/**
 	 * Receives the HTTP response when downloading a chunk.
diff --git a/Engine/Source/Runtime/Online/BuildPatchServices/Private/BuildPatchHTTP.cpp b/Engine/Source/Runtime/Online/BuildPatchServices/Private/BuildPatchHTTP.cpp
--- a/Engine/Source/Runtime/Online/BuildPatchServices/Private/BuildPatchHTTP.cpp
+++ b/Engine/Source/Runtime/Online/BuildPatchServices/Private/BuildPatchHTTP.cpp
@@ -140,13 +140,19 @@ FBuildPatchHTTP& FBuildPatchHTTP::Get()
 {
 	if (SingletonInstance.IsValid() == false)
 	{
+		checkf(IsInGameThread(), TEXT("FBuildPatchHTTP::Get() called for first time from secondary thread."));
 		SingletonInstance = MakeShareable(new FBuildPatchHTTP());
 		TickDelegateHandle = FTicker::GetCoreTicker().AddTicker( FTickerDelegate::CreateSP( SingletonInstance.Get(), &FBuildPatchHTTP::Tick ) );
 	}
 	return *SingletonInstance.Get();
 }
+void FBuildPatchHTTP::Initialize()
+{
+	FBuildPatchHTTP::Get();
+}
 void FBuildPatchHTTP::OnShutdown()
 {
+	checkf(IsInGameThread(), TEXT("FBuildPatchHTTP::OnShutdown() called from secondary thread."));
 	if (SingletonInstance.IsValid())
 	{
 		SingletonInstance->Tick( 0.0f );
diff --git a/Engine/Source/Runtime/Online/BuildPatchServices/Private/BuildPatchHTTP.h b/Engine/Source/Runtime/Online/BuildPatchServices/Private/BuildPatchHTTP.h
--- a/Engine/Source/Runtime/Online/BuildPatchServices/Private/BuildPatchHTTP.h
+++ b/Engine/Source/Runtime/Online/BuildPatchServices/Private/BuildPatchHTTP.h
@@ -76,6 +76,11 @@ public:
 	 */
 	static void QueueAnalyticsEvent( const FString& EventName, const TArray< FAnalyticsEventAttribute >& Attributes );
 
+	/**
+	 * Ensures that the wrapper is initialized
+	 */
+	static void Initialize();
+
 	/**
 	 * Clears any remaining queue and deletes the singleton. NB: Be weary of timing, if one of the static functions above is called after this
 	 * then the singleton will be re-created
diff --git a/Engine/Source/Runtime/Online/BuildPatchServices/Private/BuildPatchInstaller.cpp b/Engine/Source/Runtime/Online/BuildPatchServices/Private/BuildPatchInstaller.cpp
--- a/Engine/Source/Runtime/Online/BuildPatchServices/Private/BuildPatchInstaller.cpp
+++ b/Engine/Source/Runtime/Online/BuildPatchServices/Private/BuildPatchInstaller.cpp
@@ -524,6 +524,31 @@ bool FBuildPatchInstaller::RunInstallation(TArray<FString>& CorruptFiles)
 	return !FBuildPatchInstallError::HasFatalError();
 }
 
+
+void FBuildPatchInstaller::CleanupEmptyDirectories(const FString& RootDirectory)
+{
+	TArray<FString> SubDirNames;
+	IFileManager::Get().FindFiles(SubDirNames, *(RootDirectory / TEXT("*")), false, true);
+	for(auto DirName : SubDirNames)
+	{
+		CleanupEmptyDirectories(*(RootDirectory / DirName));
+	}
+
+	TArray<FString> SubFileNames;
+	IFileManager::Get().FindFilesRecursive(SubFileNames, *RootDirectory, TEXT("*.*"), true, false);
+	if (SubFileNames.Num() == 0)
+	{
+#if PLATFORM_MAC
+		// On Mac we need to delete the .DS_Store file, but FindFiles() skips .DS_Store files.
+		IFileManager::Get().Delete(*(RootDirectory / TEXT(".DS_Store")), false, true);
+#endif
+
+		bool bDeleteSuccess = IFileManager::Get().DeleteDirectory(*RootDirectory, false, true);
+		const uint32 LastError = FPlatformMisc::GetLastError();
+		GLog->Logf(TEXT("BuildPatchServices: Deleted Empty Folder (%u,%u) %s"), bDeleteSuccess ? 1 : 0, LastError, *RootDirectory);
+	}
+}
+
 bool FBuildPatchInstaller::RunBackupAndMove()
 {
 	GLog->Logf(TEXT("BuildPatchServices: Running backup and stage relocation"));
@@ -548,7 +573,7 @@ bool FBuildPatchInstaller::RunBackupAndMove()
 			const uint32 LastError = FPlatformMisc::GetLastError();
 			GLog->Logf(TEXT("BuildPatchServices: Removed (%u,%u) %s"), bDeleteSuccess ? 1 : 0, LastError, *OldFilename);
 		}
-
+		
 		// Now handle files that have been constructed
 		bool bSavedMoveMarkerFile = false;
 		TArray< FString > ConstructionFiles;
@@ -621,6 +646,9 @@ bool FBuildPatchInstaller::RunBackupAndMove()
 				BuildProgress.SetStateProgress(EBuildPatchProgress::MovingToInstall, FileIndexFloat / NumConstructionFilesFloat);
 			}
 		}
+		
+		// After we've completed deleting/moving patch files to the install directory, clean up any empty directories left over
+		CleanupEmptyDirectories(InstallDirectory);
 
 		bMoveSuccess = bMoveSuccess && (FBuildPatchInstallError::HasFatalError() == false);
 		if (bMoveSuccess)
diff --git a/Engine/Source/Runtime/Online/BuildPatchServices/Private/BuildPatchInstaller.h b/Engine/Source/Runtime/Online/BuildPatchServices/Private/BuildPatchInstaller.h
--- a/Engine/Source/Runtime/Online/BuildPatchServices/Private/BuildPatchInstaller.h
+++ b/Engine/Source/Runtime/Online/BuildPatchServices/Private/BuildPatchInstaller.h
@@ -234,4 +234,10 @@ private:
 	 * @param FileManifest	The manifest for the file
 	 */
 	void SetupFileAttributes( const FString& FilePath, const FFileManifestData& FileManifest );
+
+	/**
+	 * Delete empty directories from an installation
+	 * @param RootDirectory	 Root Directory for search
+	 */
+	void CleanupEmptyDirectories( const FString& RootDirectory );
 };
diff --git a/Engine/Source/Runtime/Online/BuildPatchServices/Private/BuildPatchServicesModule.cpp b/Engine/Source/Runtime/Online/BuildPatchServices/Private/BuildPatchServicesModule.cpp
--- a/Engine/Source/Runtime/Online/BuildPatchServices/Private/BuildPatchServicesModule.cpp
+++ b/Engine/Source/Runtime/Online/BuildPatchServices/Private/BuildPatchServicesModule.cpp
@@ -141,6 +141,8 @@ IBuildInstallerPtr FBuildPatchServicesModule::StartBuildInstall( IBuildManifestP
 	{
 		return NULL;
 	}
+	// Make sure the http wrapper is already created
+	FBuildPatchHTTP::Initialize();
 	// Run the install thread
 	BuildPatchInstallers.Add( MakeShareable( new FBuildPatchInstaller( OnCompleteDelegate, CurrentManifestInternal, InstallManifestInternal.ToSharedRef(), InstallDirectory, GetStagingDirectory(), InstallationInfo ) ) );
 	return BuildPatchInstallers.Top();
@@ -195,6 +197,11 @@ bool FBuildPatchServicesModule::CompactifyCloudDirectory(const TArray<FString>&
 	return FBuildDataCompactifier::CompactifyCloudDirectory(ManifestsToKeep, DataAgeThreshold, bPreview, bNoPatchDelete);
 }
 
+bool FBuildPatchServicesModule::EnumerateManifestData(FString ManifestFilePath, FString OutputFile)
+{
+	return FBuildDataEnumeration::EnumerateManifestData(MoveTemp(ManifestFilePath), MoveTemp(OutputFile));
+}
+
 #endif //WITH_BUILDPATCHGENERATION
 
 void FBuildPatchServicesModule::SetStagingDirectory( const FString& StagingDir )
diff --git a/Engine/Source/Runtime/Online/BuildPatchServices/Private/BuildPatchServicesModule.h b/Engine/Source/Runtime/Online/BuildPatchServices/Private/BuildPatchServicesModule.h
--- a/Engine/Source/Runtime/Online/BuildPatchServices/Private/BuildPatchServicesModule.h
+++ b/Engine/Source/Runtime/Online/BuildPatchServices/Private/BuildPatchServicesModule.h
@@ -86,6 +86,7 @@ public:
 	virtual bool GenerateChunksManifestFromDirectory( const FBuildPatchSettings& Settings ) override;
 	virtual bool GenerateFilesManifestFromDirectory( const FBuildPatchSettings& Settings ) override;
 	virtual bool CompactifyCloudDirectory( const TArray<FString>& ManifestsToKeep, const float DataAgeThreshold, const ECompactifyMode::Type Mode ) override;
+	virtual bool EnumerateManifestData(FString ManifestFilePath, FString OutputFile) override;
 #endif
 	virtual IBuildManifestPtr MakeManifestFromJSON( const FString& ManifestJSON ) override;
 	// End IBuildPatchServicesModule interface
diff --git a/Engine/Source/Runtime/Online/BuildPatchServices/Private/BuildPatchServicesPrivatePCH.h b/Engine/Source/Runtime/Online/BuildPatchServices/Private/BuildPatchServicesPrivatePCH.h
--- a/Engine/Source/Runtime/Online/BuildPatchServices/Private/BuildPatchServicesPrivatePCH.h
+++ b/Engine/Source/Runtime/Online/BuildPatchServices/Private/BuildPatchServicesPrivatePCH.h
@@ -42,6 +42,7 @@
 #include "BuildPatchManifest.h"
 #include "BuildPatchChunkCache.h"
 #include "BuildPatchCompactifier.h"
+#include "BuildPatchDataEnumeration.h"
 #include "BuildPatchFileConstructor.h"
 #include "BuildPatchDownloader.h"
 #include "BuildPatchInstaller.h"
diff --git a/Engine/Source/Runtime/Online/BuildPatchServices/Private/BuildPatchUtil.cpp b/Engine/Source/Runtime/Online/BuildPatchServices/Private/BuildPatchUtil.cpp
--- a/Engine/Source/Runtime/Online/BuildPatchServices/Private/BuildPatchUtil.cpp
+++ b/Engine/Source/Runtime/Online/BuildPatchServices/Private/BuildPatchUtil.cpp
@@ -8,19 +8,19 @@
 
 /* FBuildPatchUtils implementation
 *****************************************************************************/
-const FString FBuildPatchUtils::GetChunkNewFilename( const EBuildPatchAppManifestVersion::Type ManifestVersion, const FString& RootDirectory, const FGuid& ChunkGUID, const uint64& ChunkHash )
+FString FBuildPatchUtils::GetChunkNewFilename( const EBuildPatchAppManifestVersion::Type ManifestVersion, const FString& RootDirectory, const FGuid& ChunkGUID, const uint64& ChunkHash )
 {
 	check( ChunkGUID.IsValid() );
 	return FPaths::Combine( *RootDirectory, *FString::Printf( TEXT("%s/%02d/%016llX_%s.chunk"), *EBuildPatchAppManifestVersion::GetChunkSubdir( ManifestVersion ),FCrc::MemCrc32( &ChunkGUID, sizeof( FGuid ) ) % 100, ChunkHash, *ChunkGUID.ToString() ) );
 }
 
-const FString FBuildPatchUtils::GetFileNewFilename(const EBuildPatchAppManifestVersion::Type ManifestVersion, const FString& RootDirectory, const FGuid& FileGUID, const FSHAHashData& FileHash)
+FString FBuildPatchUtils::GetFileNewFilename(const EBuildPatchAppManifestVersion::Type ManifestVersion, const FString& RootDirectory, const FGuid& FileGUID, const FSHAHashData& FileHash)
 {
 	check( FileGUID.IsValid() );
 	return FPaths::Combine( *RootDirectory, *FString::Printf( TEXT("%s/%02d/%s_%s.file"), *EBuildPatchAppManifestVersion::GetFileSubdir( ManifestVersion ), FCrc::MemCrc32( &FileGUID, sizeof( FGuid ) ) % 100, *FileHash.ToString(), *FileGUID.ToString() ) );
 }
 
-const FString FBuildPatchUtils::GetFileNewFilename(const EBuildPatchAppManifestVersion::Type ManifestVersion, const FString& RootDirectory, const FGuid& FileGUID, const uint64& FileHash)
+FString FBuildPatchUtils::GetFileNewFilename(const EBuildPatchAppManifestVersion::Type ManifestVersion, const FString& RootDirectory, const FGuid& FileGUID, const uint64& FileHash)
 {
 	check( FileGUID.IsValid() );
 	return FPaths::Combine( *RootDirectory, *FString::Printf( TEXT("%s/%02d/%016llX_%s.file"), *EBuildPatchAppManifestVersion::GetFileSubdir( ManifestVersion ), FCrc::MemCrc32( &FileGUID, sizeof( FGuid ) ) % 100, FileHash, *FileGUID.ToString() ) );
@@ -52,19 +52,19 @@ void FBuildPatchUtils::GetFileDetailFromNewFilename(const FString& FileNewFilena
 	FGuid::Parse( GuidString, FileGUID );
 }
 
-const FString FBuildPatchUtils::GetChunkOldFilename( const FString& RootDirectory, const FGuid& ChunkGUID )
+FString FBuildPatchUtils::GetChunkOldFilename( const FString& RootDirectory, const FGuid& ChunkGUID )
 {
 	check( ChunkGUID.IsValid() );
 	return FPaths::Combine( *RootDirectory, *FString::Printf( TEXT("Chunks/%02d/%s.chunk"), FCrc::MemCrc_DEPRECATED( &ChunkGUID, sizeof( FGuid ) ) % 100, *ChunkGUID.ToString() ) );
 }
 
-const FString FBuildPatchUtils::GetFileOldFilename( const FString& RootDirectory, const FGuid& FileGUID )
+FString FBuildPatchUtils::GetFileOldFilename( const FString& RootDirectory, const FGuid& FileGUID )
 {
 	check( FileGUID.IsValid() );
 	return FPaths::Combine( *RootDirectory, *FString::Printf( TEXT("Files/%02d/%s.file"), FCrc::MemCrc_DEPRECATED( &FileGUID, sizeof( FGuid ) ) % 100, *FileGUID.ToString() ) );
 }
 
-const FString FBuildPatchUtils::GetDataTypeOldFilename( const FBuildPatchData::Type DataType, const FString& RootDirectory, const FGuid& Guid )
+FString FBuildPatchUtils::GetDataTypeOldFilename( const FBuildPatchData::Type DataType, const FString& RootDirectory, const FGuid& Guid )
 {
 	check( Guid.IsValid() );
 
@@ -81,7 +81,7 @@ const FString FBuildPatchUtils::GetDataTypeOldFilename( const FBuildPatchData::T
 	return TEXT( "" );
 }
 
-const FString FBuildPatchUtils::GetDataFilename(const FBuildPatchAppManifestRef& Manifest, const FString& RootDirectory, const FGuid& DataGUID)
+FString FBuildPatchUtils::GetDataFilename(const FBuildPatchAppManifestRef& Manifest, const FString& RootDirectory, const FGuid& DataGUID)
 {
 	const FBuildPatchData::Type DataType = Manifest->IsFileDataManifest() ? FBuildPatchData::FileData : FBuildPatchData::ChunkData;
 	if (Manifest->GetManifestVersion() < EBuildPatchAppManifestVersion::DataFileRenames)
diff --git a/Engine/Source/Runtime/Online/BuildPatchServices/Private/BuildPatchUtil.h b/Engine/Source/Runtime/Online/BuildPatchServices/Private/BuildPatchUtil.h
--- a/Engine/Source/Runtime/Online/BuildPatchServices/Private/BuildPatchUtil.h
+++ b/Engine/Source/Runtime/Online/BuildPatchServices/Private/BuildPatchUtil.h
@@ -19,7 +19,7 @@ struct FBuildPatchUtils
 	 * @param ChunkHash			The chunk rolling hash value
 	 * @return	the chunk path
 	 */
-	static const FString GetChunkNewFilename(const EBuildPatchAppManifestVersion::Type ManifestVersion, const FString& RootDirectory, const FGuid& ChunkGUID, const uint64& ChunkHash);
+	static FString GetChunkNewFilename(const EBuildPatchAppManifestVersion::Type ManifestVersion, const FString& RootDirectory, const FGuid& ChunkGUID, const uint64& ChunkHash);
 
 	/**
 	 * Gets the filename for a file chunk generated from it's GUID and Hash, which is the new format.
@@ -29,8 +29,8 @@ struct FBuildPatchUtils
 	 * @param FileHash			The file hash value
 	 * @return	The file chunk path
 	 */
-	static const FString GetFileNewFilename(const EBuildPatchAppManifestVersion::Type ManifestVersion, const FString& RootDirectory, const FGuid& FileGUID, const FSHAHashData& FileHash);
-	static const FString GetFileNewFilename(const EBuildPatchAppManifestVersion::Type ManifestVersion, const FString& RootDirectory, const FGuid& FileGUID, const uint64& FilePartHash);
+	static FString GetFileNewFilename(const EBuildPatchAppManifestVersion::Type ManifestVersion, const FString& RootDirectory, const FGuid& FileGUID, const FSHAHashData& FileHash);
+	static FString GetFileNewFilename(const EBuildPatchAppManifestVersion::Type ManifestVersion, const FString& RootDirectory, const FGuid& FileGUID, const uint64& FilePartHash);
 
 	/**
 	 * Gets the chunk GUID and Hash, from the filename, which is the new format.
@@ -54,7 +54,7 @@ struct FBuildPatchUtils
 	 * @param ChunkGuid			The chunk Guid
 	 * @return	the chunk path
 	 */
-	static const FString GetChunkOldFilename(const FString& RootDirectory, const FGuid& ChunkGUID);
+	static FString GetChunkOldFilename(const FString& RootDirectory, const FGuid& ChunkGUID);
 
 	/**
 	 * Gets the filename for a file data part generated from it's GUID
@@ -62,7 +62,7 @@ struct FBuildPatchUtils
 	 * @param FileGUID			The file part Guid
 	 * @return	The file data path
 	 */
-	static const FString GetFileOldFilename(const FString& RootDirectory, const FGuid& FileGUID);
+	static FString GetFileOldFilename(const FString& RootDirectory, const FGuid& FileGUID);
 
 	/**
 	 * Gets the filename for a specific data part type from it's GUID
@@ -71,7 +71,7 @@ struct FBuildPatchUtils
 	 * @param DataGUID			The data Guid
 	 * @return	the data part path
 	 */
-	static const FString GetDataTypeOldFilename(const FBuildPatchData::Type DataType, const FString& RootDirectory, const FGuid& Guid);
+	static FString GetDataTypeOldFilename(const FBuildPatchData::Type DataType, const FString& RootDirectory, const FGuid& Guid);
 
 	/**
 	 * Gets the filename for any data part. Wraps the choice between all of the above
@@ -80,7 +80,7 @@ struct FBuildPatchUtils
 	 * @param DataGUID			The data Guid
 	 * @return	the data part path
 	 */
-	static const FString GetDataFilename(const FBuildPatchAppManifestRef& Manifest, const FString& RootDirectory, const FGuid& DataGUID);
+	static FString GetDataFilename(const FBuildPatchAppManifestRef& Manifest, const FString& RootDirectory, const FGuid& DataGUID);
 
 	/**
 	 * Gets the GUID for a data file according to it's filename (new or old)
diff --git a/Engine/Source/Runtime/Online/BuildPatchServices/Public/Interfaces/IBuildPatchServicesModule.h b/Engine/Source/Runtime/Online/BuildPatchServices/Public/Interfaces/IBuildPatchServicesModule.h
--- a/Engine/Source/Runtime/Online/BuildPatchServices/Public/Interfaces/IBuildPatchServicesModule.h
+++ b/Engine/Source/Runtime/Online/BuildPatchServices/Public/Interfaces/IBuildPatchServicesModule.h
@@ -136,6 +136,15 @@ public:
 	 * @return              true if no file errors occurred
 	 */
 	virtual bool CompactifyCloudDirectory(const TArray<FString>& ManifestsToKeep, const float DataAgeThreshold, const ECompactifyMode::Type Mode) = 0;
+
+	/**
+	 * Saves to file, and logs, a full list of cloud dir relative referenced data file paths
+	 * @param   ManifestFilePath     A full file path for the manifest to be loaded
+	 * @param   OutputFile           A full file path where to save the output text
+	 * @return  true if successful
+	 */
+	virtual bool EnumerateManifestData(FString ManifestFilePath, FString OutputFile) = 0;
+
 #endif // WITH_BUILDPATCHGENERATION
 
 	/**
diff --git a/Engine/Source/Runtime/Online/HTTP/Private/Curl/CurlHttp.cpp b/Engine/Source/Runtime/Online/HTTP/Private/Curl/CurlHttp.cpp
--- a/Engine/Source/Runtime/Online/HTTP/Private/Curl/CurlHttp.cpp
+++ b/Engine/Source/Runtime/Online/HTTP/Private/Curl/CurlHttp.cpp
@@ -20,6 +20,7 @@ FCurlHttpRequest::FCurlHttpRequest(CURLM * InMultiHandle)
 	,	BytesSent(0)
 	,	CompletionStatus(EHttpRequestStatus::NotStarted)
 	,	ElapsedTime(0.0f)
+	,	TimeSinceLastResponse(0.0f)
 {
 	check(MultiHandle);
 	if (MultiHandle)
@@ -174,22 +175,7 @@ void FCurlHttpRequest::SetVerb(const FString& InVerb)
 void FCurlHttpRequest::SetURL(const FString& InURL)
 {
 	check(EasyHandle);
-
-	URL = InURL.Replace(TEXT("%"), TEXT("%25"));
-	URL = URL.Replace(TEXT(" "), TEXT("%20"));
-	URL = URL.Replace(TEXT("\""), TEXT("%22"));
-	URL = URL.Replace(TEXT("<"), TEXT("%3C"));
-	URL = URL.Replace(TEXT(">"), TEXT("%3E"));
-
-	URL = URL.Replace(TEXT("["), TEXT("%5B"));
-	URL = URL.Replace(TEXT("]"), TEXT("%5D"));
-	URL = URL.Replace(TEXT("\\"), TEXT("%5C"));
-	URL = URL.Replace(TEXT("^"), TEXT("%5E"));
-
-	URL = URL.Replace(TEXT("`"), TEXT("%60"));
-	URL = URL.Replace(TEXT("{"), TEXT("%7B"));
-	URL = URL.Replace(TEXT("}"), TEXT("%7D"));
-	URL = URL.Replace(TEXT("|"), TEXT("%7C"));
+	URL = InURL;
 }
 
 void FCurlHttpRequest::SetContent(const TArray<uint8>& ContentPayload)
@@ -262,6 +248,8 @@ size_t FCurlHttpRequest::ReceiveResponseHeaderCallback(void* Ptr, size_t SizeInB
 
 	if (Response.IsValid())
 	{
+		TimeSinceLastResponse = 0.0f;
+
 		uint32 HeaderSize = SizeInBlocks * BlockSizeInBytes;
 		if (HeaderSize > 0 && HeaderSize <= CURL_MAX_HTTP_HEADER)
 		{
@@ -304,6 +292,8 @@ size_t FCurlHttpRequest::ReceiveResponseBodyCallback(void* Ptr, size_t SizeInBlo
 
 	if (Response.IsValid())
 	{
+		TimeSinceLastResponse = 0.0f;
+
 		uint32 SizeToDownload = SizeInBlocks * BlockSizeInBytes;
 
 		UE_LOG(LogHttp, Verbose, TEXT("%p: ReceiveResponseBodyCallback: %d bytes out of %d received. (SizeInBlocks=%d, BlockSizeInBytes=%d, Response->TotalBytesRead=%d, Response->GetContentLength()=%d, SizeToDownload=%d (<-this will get returned from the callback))"),
@@ -322,7 +312,7 @@ size_t FCurlHttpRequest::ReceiveResponseBodyCallback(void* Ptr, size_t SizeInBlo
 			Response->TotalBytesRead += SizeToDownload;
 
 			// Update response progress
-			OnRequestProgress().ExecuteIfBound(SharedThis(this), Response->TotalBytesRead);
+			OnRequestProgress().ExecuteIfBound(SharedThis(this), BytesSent, Response->TotalBytesRead);
 
 			return SizeToDownload;
 		}
@@ -337,11 +327,15 @@ size_t FCurlHttpRequest::ReceiveResponseBodyCallback(void* Ptr, size_t SizeInBlo
 
 size_t FCurlHttpRequest::UploadCallback(void* Ptr, size_t SizeInBlocks, size_t BlockSizeInBytes)
 {
+	TimeSinceLastResponse = 0.0f;
+
 	size_t SizeToSend = RequestPayload.Num() - BytesSent;
 	size_t SizeToSendThisTime = 0;
 
 	if (SizeToSend != 0)
 	{
+		OnRequestProgress().ExecuteIfBound(SharedThis(this), BytesSent, 0);
+
 		SizeToSendThisTime = FMath::Min(SizeToSend, SizeInBlocks * BlockSizeInBytes);
 		if (SizeToSendThisTime != 0)
 		{
@@ -444,6 +438,12 @@ bool FCurlHttpRequest::StartRequest()
 {
 	check(EasyHandle);
 
+	bCompleted = false;
+	bCanceled = false;
+
+	curl_slist_free_all(HeaderList);
+	HeaderList = nullptr;
+
 	// default no verb to a GET
 	if (Verb.IsEmpty())
 	{
@@ -638,24 +638,32 @@ const FHttpResponsePtr FCurlHttpRequest::GetResponse() const
 
 void FCurlHttpRequest::Tick(float DeltaSeconds)
 {
+	// keep track of elapsed seconds
+	ElapsedTime += DeltaSeconds;
+	TimeSinceLastResponse += DeltaSeconds;
+
 	// check for true completion/cancellation
-	if (bCompleted || bCanceled)
+	if (bCompleted && 
+		ElapsedTime >= FHttpModule::Get().GetHttpDelayTime())
 	{
 		FinishedRequest();
-		return;
 	}
-
-	// keep track of elapsed seconds
-	ElapsedTime += DeltaSeconds;
-	const float HttpTimeout = FHttpModule::Get().GetHttpTimeout();
-	if (HttpTimeout > 0 && ElapsedTime >= HttpTimeout)
+	else if (bCanceled)
 	{
-		UE_LOG(LogHttp, Warning, TEXT("Timeout processing Http request. %p"),
-			this);
-
-		// finish it off since it is timeout
 		FinishedRequest();
 	}
+	else
+	{
+		const float HttpTimeout = FHttpModule::Get().GetHttpTimeout();
+		if (HttpTimeout > 0 && TimeSinceLastResponse >= HttpTimeout)
+		{
+			UE_LOG(LogHttp, Warning, TEXT("Timeout processing Http request. %p"),
+				this);
+
+			// finish it off since it is timeout
+			FinishedRequest();
+		}
+	}
 }
 
 void FCurlHttpRequest::FinishedRequest()
diff --git a/Engine/Source/Runtime/Online/HTTP/Private/Curl/CurlHttp.h b/Engine/Source/Runtime/Online/HTTP/Private/Curl/CurlHttp.h
--- a/Engine/Source/Runtime/Online/HTTP/Private/Curl/CurlHttp.h
+++ b/Engine/Source/Runtime/Online/HTTP/Private/Curl/CurlHttp.h
@@ -306,7 +306,7 @@ private:
 	TArray<uint8> RequestPayload;
 	/** Delegate that will get called once request completes or on any error */
 	FHttpRequestCompleteDelegate RequestCompleteDelegate;
-	/** Delegate that will get called once per tick with bytes downloaded so far */
+	/** Delegate that will get called once per tick with total bytes uploaded and downloaded so far */
 	FHttpRequestProgressDelegate RequestProgressDelegate;
 	/** Current status of request being processed */
 	EHttpRequestStatus::Type CompletionStatus;
@@ -314,6 +314,8 @@ private:
 	TMap<FString, FString> Headers;
 	/** Total elapsed time in seconds since the start of the request */
 	float ElapsedTime;
+	/** Elapsed time since the last received HTTP response. */
+	float TimeSinceLastResponse;
 };
 
 
diff --git a/Engine/Source/Runtime/Online/HTTP/Private/HTML5/HTML5HTTP.cpp b/Engine/Source/Runtime/Online/HTTP/Private/HTML5/HTML5HTTP.cpp
--- a/Engine/Source/Runtime/Online/HTTP/Private/HTML5/HTML5HTTP.cpp
+++ b/Engine/Source/Runtime/Online/HTTP/Private/HTML5/HTML5HTTP.cpp
@@ -257,12 +257,12 @@ void FHTML5HttpRequest::ProgressCallback(void* arg, int Loaded, int Total) {
 		if (Response.IsValid())
 		{
 			Response->TotalBytesRead = Loaded;
-			OnRequestProgress().ExecuteIfBound(SharedThis(this), Response->TotalBytesRead);
+			OnRequestProgress().ExecuteIfBound(SharedThis(this), 0, Response->TotalBytesRead);
 		}
 	}
 	else {
 		BytesSent = Loaded;
-		OnRequestProgress().ExecuteIfBound(SharedThis(this), BytesSent);
+		OnRequestProgress().ExecuteIfBound(SharedThis(this), BytesSent, 0);
 	}
 
 	UE_LOG(LogHttp, Verbose, TEXT("Loaded: %d, Total: %d"), Loaded, Total);
diff --git a/Engine/Source/Runtime/Online/HTTP/Private/HttpModule.cpp b/Engine/Source/Runtime/Online/HTTP/Private/HttpModule.cpp
--- a/Engine/Source/Runtime/Online/HTTP/Private/HttpModule.cpp
+++ b/Engine/Source/Runtime/Online/HTTP/Private/HttpModule.cpp
@@ -1,6 +1,7 @@
 // Copyright 1998-2015 Epic Games, Inc. All Rights Reserved.
 
 #include "HttpPrivatePCH.h"
+#include "NullHttp.h"
 #include "HttpTests.h"
 
 DEFINE_LOG_CATEGORY(LogHttp);
@@ -42,6 +43,12 @@ void FHttpModule::StartupModule()
 	
 	bEnableHttp = true;
 	GConfig->GetBool(TEXT("HTTP"), TEXT("bEnableHttp"), bEnableHttp, GEngineIni);
+
+	bUseNullHttp = false;
+	GConfig->GetBool(TEXT("HTTP"), TEXT("bUseNullHttp"), bUseNullHttp, GEngineIni);
+
+	HttpDelayTime = 0;
+	GConfig->GetFloat(TEXT("HTTP"), TEXT("HttpDelayTime"), HttpDelayTime, GEngineIni);
 }
 
 void FHttpModule::ShutdownModule()
@@ -121,7 +128,14 @@ FHttpModule& FHttpModule::Get()
 
 TSharedRef<IHttpRequest> FHttpModule::CreateRequest()
 {
-	// Create the platform specific Http request instance
-	return TSharedRef<IHttpRequest>( FPlatformHttp::ConstructRequest() );
+	if (bUseNullHttp)
+	{
+		return TSharedRef<IHttpRequest>(new FNullHttpRequest());
+	}
+	else
+	{
+		// Create the platform specific Http request instance
+		return TSharedRef<IHttpRequest>(FPlatformHttp::ConstructRequest());
+	}
 }
 
diff --git a/Engine/Source/Runtime/Online/HTTP/Private/HttpRetrySystem.cpp b/Engine/Source/Runtime/Online/HTTP/Private/HttpRetrySystem.cpp
new file mode 100644
--- /dev/null
+++ b/Engine/Source/Runtime/Online/HTTP/Private/HttpRetrySystem.cpp
@@ -0,0 +1,347 @@
+// Copyright 1998-2014 Epic Games, Inc. All Rights Reserved.
+
+#include "HttpPrivatePCH.h"
+
+#include "HttpRetrySystem.h"
+
+FHttpRetrySystem::FRequest::FRequest(TSharedRef<IHttpRequest>& HttpRequest, FHttpRetrySystem::FRetryLimitCountSetting InRetryLimitCountOverride, FHttpRetrySystem::FRetryTimeoutRelativeSecondsSetting InRetryTimeoutRelativeSecondsOverride)
+    : FHttpRequestAdapterBase(HttpRequest)
+    , Status(FHttpRetrySystem::FRequest::EStatus::NotStarted)
+    , RetryLimitCountOverride(InRetryLimitCountOverride)
+    , RetryTimeoutRelativeSecondsOverride(InRetryTimeoutRelativeSecondsOverride)
+{
+    // if the InRetryTimeoutRelativeSecondsOverride override is being used the value cannot be negative
+    check(!(InRetryTimeoutRelativeSecondsOverride.bUseValue) || (InRetryTimeoutRelativeSecondsOverride.Value >= 0.0));
+}
+
+FHttpRetrySystem::FManager::FManager(FHttpRetrySystem::FRetryLimitCountSetting InRetryLimitCountDefault, FHttpRetrySystem::FRetryTimeoutRelativeSecondsSetting InRetryTimeoutRelativeSecondsDefault)
+    : RandomFailureRate(FRandomFailureRateSetting::Unused())
+    , RetryLimitCountDefault(InRetryLimitCountDefault)
+    , RetryTimeoutRelativeSecondsDefault(InRetryTimeoutRelativeSecondsDefault)
+{}
+
+bool FHttpRetrySystem::FManager::ShouldRetry(const FHttpRetryRequestEntry& HttpRetryRequestEntry)
+{
+    bool bResult = false;
+
+    FHttpResponsePtr Response = HttpRetryRequestEntry.HttpRequest->GetResponse();
+    if (!Response.IsValid() || Response->GetResponseCode() == 0)
+    {
+        bResult = true;
+    }
+
+    return bResult;
+}
+
+bool FHttpRetrySystem::FManager::CanRetry(const FHttpRetryRequestEntry& HttpRetryRequestEntry)
+{
+    bool bResult = false;
+
+    bool bShouldTestCurrentRetryCount = false;
+    double RetryLimitCount = 0;
+    if (HttpRetryRequestEntry.HttpRequest->RetryLimitCountOverride.bUseValue)
+    {
+        bShouldTestCurrentRetryCount = true;
+        RetryLimitCount = HttpRetryRequestEntry.HttpRequest->RetryLimitCountOverride.Value;
+    }
+    else if (RetryLimitCountDefault.bUseValue)
+    {
+        bShouldTestCurrentRetryCount = true;
+        RetryLimitCount = RetryLimitCountDefault.Value;
+    }
+
+    if (bShouldTestCurrentRetryCount)
+    {
+        if (HttpRetryRequestEntry.CurrentRetryCount < RetryLimitCount)
+        {
+            bResult = true;
+        }
+    }
+
+    return bResult;
+}
+
+bool FHttpRetrySystem::FManager::HasTimedOut(const FHttpRetryRequestEntry& HttpRetryRequestEntry, const double NowAbsoluteSeconds)
+{
+    bool bResult = false;
+
+    bool bShouldTestRetryTimeout = false;
+    double RetryTimeoutAbsoluteSeconds = HttpRetryRequestEntry.RequestStartTimeAbsoluteSeconds;
+    if (HttpRetryRequestEntry.HttpRequest->RetryTimeoutRelativeSecondsOverride.bUseValue)
+    {
+        bShouldTestRetryTimeout = true;
+        RetryTimeoutAbsoluteSeconds += HttpRetryRequestEntry.HttpRequest->RetryTimeoutRelativeSecondsOverride.Value;
+    }
+    else if (RetryTimeoutRelativeSecondsDefault.bUseValue)
+    {
+        bShouldTestRetryTimeout = true;
+        RetryTimeoutAbsoluteSeconds += RetryTimeoutRelativeSecondsDefault.Value;
+    }
+
+    if (bShouldTestRetryTimeout)
+    {
+        if (NowAbsoluteSeconds >= RetryTimeoutAbsoluteSeconds)
+        {
+            bResult = true;
+        }
+    }
+
+    return bResult;
+}
+
+float FHttpRetrySystem::FManager::GetLockoutPeriodSeconds(const FHttpRetryRequestEntry& HttpRetryRequestEntry)
+{
+    float lockoutTime = 0.0f;
+
+    if(HttpRetryRequestEntry.CurrentRetryCount >= 1)
+    {
+        lockoutTime = 5.0f + 5.0f * ((HttpRetryRequestEntry.CurrentRetryCount - 1) >> 1);
+        lockoutTime = lockoutTime > 30.0f ? 30.0f : lockoutTime;
+    }
+
+    return lockoutTime;
+}
+
+static FRandomStream temp(4435261);
+
+bool FHttpRetrySystem::FManager::Update(uint32* FileCount, uint32* FailingCount, uint32* FailedCount, uint32* CompletedCount)
+{
+	bool bIsGreen = true;
+
+	if (FileCount != nullptr)
+	{
+		*FileCount = RequestList.Num();
+	}
+
+	const double NowAbsoluteSeconds = FPlatformTime::Seconds();
+
+	// Basic algorithm
+	// for each managed item
+	//    if the item hasn't timed out
+	//       if the item's retry state is NotStarted
+	//          if the item's request's state is not NotStarted
+	//             move the item's retry state to Processing
+	//          endif
+	//       endif
+	//       if the item's retry state is Processing
+	//          if the item's request's state is Failed
+	//             flag return code to false
+	//             if the item can be retried
+	//                increment FailingCount if applicable
+	//                retry the item's request
+	//                increment the item's retry count
+	//             else
+	//                increment FailedCount if applicable
+	//                set the item's retry state to FailedRetry
+	//             endif
+	//          else if the item's request's state is Succeeded
+	//          endif
+	//       endif
+	//    else
+	//       flag return code to false
+	//       set the item's retry state to FailedTimeout
+	//       increment FailedCount if applicable
+	//    endif
+	//    if the item's retry state is FailedRetry
+	//       do stuff
+	//    endif
+	//    if the item's retry state is FailedTimeout
+	//       do stuff
+	//    endif
+	//    if the item's retry state is Succeeded
+	//       do stuff
+	//    endif
+	// endfor
+
+	int32 index = 0;
+	while (index < RequestList.Num())
+	{
+		FHttpRetryRequestEntry& HttpRetryRequestEntry = RequestList[index];
+		TSharedRef<FHttpRetrySystem::FRequest>& HttpRetryRequest = HttpRetryRequestEntry.HttpRequest;
+
+		EHttpRequestStatus::Type RequestStatus = HttpRetryRequest->HttpRequest->GetStatus();
+
+        if (!HasTimedOut(HttpRetryRequestEntry, NowAbsoluteSeconds))
+		{
+            if (HttpRetryRequest->Status == FHttpRetrySystem::FRequest::EStatus::NotStarted)
+			{
+				if (RequestStatus != EHttpRequestStatus::NotStarted)
+				{
+                    HttpRetryRequest->Status = FHttpRetrySystem::FRequest::EStatus::Processing;
+				}
+			}
+
+            if (HttpRetryRequest->Status == FHttpRetrySystem::FRequest::EStatus::Processing)
+			{
+				bool forceFail = false;
+
+                // Code to simulate request failure
+                if (RequestStatus == EHttpRequestStatus::Succeeded && RandomFailureRate.bUseValue)
+				{
+					float random = temp.GetFraction();
+                    if (random < RandomFailureRate.Value)
+					{
+						forceFail = true;
+					}
+				}
+
+                if ((RequestStatus == EHttpRequestStatus::Failed) || forceFail)
+				{
+					bIsGreen = false;
+                    if(HttpRetryRequestEntry.bShouldCancel == false)
+                    {
+                        if (forceFail || (ShouldRetry(HttpRetryRequestEntry) && CanRetry(HttpRetryRequestEntry)))
+					    {
+                            float lockoutPeriod = GetLockoutPeriodSeconds(HttpRetryRequestEntry);
+
+                            if(lockoutPeriod > 0.0f)
+                            {
+                                UE_LOG(LogHttp, Warning, TEXT("Lockout of %fs on %s"), lockoutPeriod, *(HttpRetryRequest->GetURL()));
+                            }
+
+                            HttpRetryRequestEntry.LockoutEndTimeAbsoluteSeconds = NowAbsoluteSeconds + lockoutPeriod;
+                            HttpRetryRequest->Status = FHttpRetrySystem::FRequest::EStatus::ProcessingLockout;
+					    }
+					    else
+					    {
+						    UE_LOG(LogHttp, Warning, TEXT("Retry exhausted on %s"), *(HttpRetryRequest->GetURL()));
+						    if (FailedCount != nullptr)
+						    {
+							    ++(*FailedCount);
+						    }
+                            HttpRetryRequest->Status = FHttpRetrySystem::FRequest::EStatus::FailedRetry;
+					    }
+                    }
+                    else
+                    {
+                        UE_LOG(LogHttp, Warning, TEXT("Request cancelled on %s"), *(HttpRetryRequest->GetURL()));
+                        HttpRetryRequest->Status = FHttpRetrySystem::FRequest::EStatus::Cancelled;
+                    }
+                }
+				else if (RequestStatus == EHttpRequestStatus::Succeeded)
+				{
+					if (HttpRetryRequestEntry.CurrentRetryCount > 0)
+					{
+						UE_LOG(LogHttp, Warning, TEXT("Success on %s"), *(HttpRetryRequest->GetURL()));
+					}
+
+					if (CompletedCount != nullptr)
+					{
+						++(*CompletedCount);
+					}
+
+                    HttpRetryRequest->Status = FHttpRetrySystem::FRequest::EStatus::Succeeded;
+				}
+			}
+
+            if (HttpRetryRequest->Status == FHttpRetrySystem::FRequest::EStatus::ProcessingLockout)
+            {
+                if (NowAbsoluteSeconds >= HttpRetryRequestEntry.LockoutEndTimeAbsoluteSeconds)
+                {
+                    // if this fails the HttpRequest's state will be failed which will cause the retry logic to kick(as expected)
+                    bool success = HttpRetryRequest->HttpRequest->ProcessRequest();
+                    if (success)
+                    {
+                        UE_LOG(LogHttp, Warning, TEXT("Retry %d on %s"), HttpRetryRequestEntry.CurrentRetryCount + 1, *(HttpRetryRequest->GetURL()));
+
+                        ++HttpRetryRequestEntry.CurrentRetryCount;
+                        HttpRetryRequest->Status = FRequest::EStatus::Processing;
+                    }
+                }
+
+                if (FailingCount != nullptr)
+                {
+                    ++(*FailingCount);
+                }
+            }
+        }
+		else
+		{
+			UE_LOG(LogHttp, Warning, TEXT("Timeout on %s"), HttpRetryRequestEntry.CurrentRetryCount + 1, *(HttpRetryRequest->GetURL()));
+			bIsGreen = false;
+            HttpRetryRequest->Status = FHttpRetrySystem::FRequest::EStatus::FailedTimeout;
+			if (FailedCount != nullptr)
+			{
+				++(*FailedCount);
+			}
+		}
+
+		bool bWasCompleted = false;
+		bool bWasSuccessful = false;
+
+        if (HttpRetryRequest->Status == FHttpRetrySystem::FRequest::EStatus::Cancelled ||
+            HttpRetryRequest->Status == FHttpRetrySystem::FRequest::EStatus::FailedRetry ||
+            HttpRetryRequest->Status == FHttpRetrySystem::FRequest::EStatus::FailedTimeout ||
+            HttpRetryRequest->Status == FHttpRetrySystem::FRequest::EStatus::Succeeded)
+		{
+			bWasCompleted = true;
+            bWasSuccessful = HttpRetryRequest->Status == FHttpRetrySystem::FRequest::EStatus::Succeeded;
+		}
+
+		if (bWasCompleted)
+		{
+			HttpRetryRequest->OnProcessRequestComplete().ExecuteIfBound(HttpRetryRequest, bWasSuccessful);
+		}
+
+        if(bWasSuccessful)
+        {
+            if(CompletedCount != nullptr)
+            {
+                ++(*CompletedCount);
+            }
+        }
+
+		if (bWasCompleted)
+		{
+			RequestList.RemoveAtSwap(index);
+		}
+		else
+		{
+			++index;
+		}
+	}
+
+	return bIsGreen;
+}
+
+FHttpRetrySystem::FManager::FHttpRetryRequestEntry::FHttpRetryRequestEntry(TSharedRef<FHttpRetrySystem::FRequest>& InHttpRequest)
+    : bShouldCancel(false)
+    , CurrentRetryCount(0)
+    , RequestStartTimeAbsoluteSeconds(0.0)
+    , HttpRequest(InHttpRequest)
+{}
+
+bool FHttpRetrySystem::FManager::ProcessRequest(TSharedRef<FHttpRetrySystem::FRequest>& HttpRequest)
+{
+	bool bResult = HttpRequest->ProcessRequest();
+
+	if (bResult)
+	{
+        RequestList.Add(FHttpRetryRequestEntry(HttpRequest));
+	}
+
+	return bResult;
+}
+
+void FHttpRetrySystem::FManager::CancelRequest(TSharedRef<FHttpRetrySystem::FRequest>& HttpRequest)
+{
+    for (int32 i = 0; i < RequestList.Num(); ++i)
+    {
+        FHttpRetryRequestEntry& EntryRef = RequestList[i];
+
+        if(EntryRef.HttpRequest == HttpRequest)
+        {
+            EntryRef.bShouldCancel = true;
+        }
+    }
+    HttpRequest->CancelRequest();
+
+    /*
+    if (bResult)
+    {
+        RequestList.Add(FHttpRetryRequestEntry(HttpRequest));
+    }
+
+    return bResult;
+    */
+}
diff --git a/Engine/Source/Runtime/Online/HTTP/Private/Mac/MacHTTP.cpp b/Engine/Source/Runtime/Online/HTTP/Private/Mac/MacHTTP.cpp
--- a/Engine/Source/Runtime/Online/HTTP/Private/Mac/MacHTTP.cpp
+++ b/Engine/Source/Runtime/Online/HTTP/Private/Mac/MacHTTP.cpp
@@ -33,6 +33,7 @@ FMacHttpRequest::~FMacHttpRequest()
 
 FString FMacHttpRequest::GetURL()
 {
+	SCOPED_AUTORELEASE_POOL;
 	FString URL([[Request URL] absoluteString]);
 	UE_LOG(LogHttp, Verbose, TEXT("FMacHttpRequest::GetURL() - %s"), *URL);
 	return URL;
@@ -41,6 +42,7 @@ FString FMacHttpRequest::GetURL()
 
 void FMacHttpRequest::SetURL(const FString& URL)
 {
+	SCOPED_AUTORELEASE_POOL;
 	UE_LOG(LogHttp, Verbose, TEXT("FMacHttpRequest::SetURL() - %s"), *URL);
 	[Request setURL: [NSURL URLWithString: URL.GetNSString()]];
 }
@@ -157,6 +159,7 @@ void FMacHttpRequest::SetVerb(const FString& Verb)
 
 bool FMacHttpRequest::ProcessRequest()
 {
+	SCOPED_AUTORELEASE_POOL;
 	UE_LOG(LogHttp, Verbose, TEXT("FMacHttpRequest::ProcessRequest()"));
 	bool bStarted = false;
 
@@ -198,7 +201,7 @@ FHttpRequestCompleteDelegate& FMacHttpRequest::OnProcessRequestComplete()
 
 FHttpRequestProgressDelegate& FMacHttpRequest::OnRequestProgress() 
 {
-	UE_LOG(LogHttp, Verbose, TEXT("FMacHttpRequest::OnRequestProgress()"));
+	UE_LOG(LogHttp, VeryVerbose, TEXT("FMacHttpRequest::OnRequestProgress()"));
 	return RequestProgressDelegate;
 }
 
@@ -305,7 +308,6 @@ void FMacHttpRequest::CancelRequest()
 	if(Connection != NULL)
 	{
 		[Connection cancel];
-		Connection = NULL;
 	}
 	FinishedRequest();
 }
@@ -328,13 +330,13 @@ void FMacHttpRequest::Tick(float DeltaSeconds)
 {
 	if( CompletionStatus == EHttpRequestStatus::Processing || Response->HadError() )
 	{
-		if( OnRequestProgress().IsBound() )
+		if (OnRequestProgress().IsBound())
 		{
-			const int32 NumBytesReceived = Response->GetNumBytesReceived();
-			if( ProgressBytesSent < NumBytesReceived )
+			const int32 BytesWritten = Response->GetNumBytesWritten();
+			const int32 BytesRead = Response->GetNumBytesReceived();
+			if (BytesWritten > 0 || BytesRead > 0)
 			{
-				ProgressBytesSent = NumBytesReceived;
-				OnRequestProgress().Execute(SharedThis(this), NumBytesReceived);
+				OnRequestProgress().Execute(SharedThis(this), BytesWritten, BytesRead);
 			}
 		}
 		if( Response->IsReady() )
@@ -358,6 +360,7 @@ float FMacHttpRequest::GetElapsedTime()
 @synthesize Response;
 @synthesize bIsReady;
 @synthesize bHadError;
+@synthesize BytesWritten;
 
 
 -(FHttpResponseMacWrapper*) init
@@ -370,6 +373,13 @@ float FMacHttpRequest::GetElapsedTime()
 }
 
 
+-(void) connection:(NSURLConnection *)connection didSendBodyData:(NSInteger)bytesWritten totalBytesWritten:(NSInteger)totalBytesWritten totalBytesExpectedToWrite:(NSInteger)totalBytesExpectedToWrite
+{
+	UE_LOG(LogHttp, Verbose, TEXT("didSendBodyData:(NSInteger)bytesWritten totalBytes:Written:(NSInteger)totalBytesWritten totalBytesExpectedToWrite:(NSInteger)totalBytesExpectedToWrite"));
+	self.BytesWritten = totalBytesWritten;
+	UE_LOG(LogHttp, Verbose, TEXT("didSendBodyData: totalBytesWritten = %d, totalBytesExpectedToWrite = %d: %p"), totalBytesWritten, totalBytesExpectedToWrite, self);
+}
+
 -(void) connection:(NSURLConnection *)connection didReceiveResponse:(NSURLResponse *)response
 {
 	UE_LOG(LogHttp, Verbose, TEXT("didReceiveResponse:(NSURLResponse *)response"));
@@ -411,6 +421,11 @@ float FMacHttpRequest::GetElapsedTime()
 	return Payload;
 }
 
+-(int32)getBytesWritten
+{
+	return self.BytesWritten;
+}
+
 @end
 
 
@@ -581,3 +596,9 @@ const int32 FMacHttpResponse::GetNumBytesReceived() const
 {
 	return [ResponseWrapper getPayload].Num();
 }
+
+const int32 FMacHttpResponse::GetNumBytesWritten() const
+{
+    int32 NumBytesWritten = [ResponseWrapper getBytesWritten];
+    return NumBytesWritten;
+}
diff --git a/Engine/Source/Runtime/Online/HTTP/Private/Mac/MacHTTP.h b/Engine/Source/Runtime/Online/HTTP/Private/Mac/MacHTTP.h
--- a/Engine/Source/Runtime/Online/HTTP/Private/Mac/MacHTTP.h
+++ b/Engine/Source/Runtime/Online/HTTP/Private/Mac/MacHTTP.h
@@ -122,7 +122,11 @@ private:
 @property BOOL bIsReady;
 /** When the response is complete, indicates whether the response was received without error. */
 @property BOOL bHadError;
+/** The total number of bytes written out during the request/response */
+@property int32 BytesWritten;
 
+/** Delegate called when we send data. See Mac docs for when/how this should be used. */
+-(void) connection:(NSURLConnection *)connection didSendBodyData:(NSInteger)bytesWritten totalBytesWritten:(NSInteger)totalBytesWritten totalBytesExpectedToWrite:(NSInteger)totalBytesExpectedToWrite;
 /** Delegate called with we receive a response. See Mac docs for when/how this should be used. */
 -(void) connection:(NSURLConnection *)connection didReceiveResponse:(NSURLResponse *)response;
 /** Delegate called with we receive data. See Mac docs for when/how this should be used. */
@@ -133,6 +137,7 @@ private:
 -(void) connectionDidFinishLoading:(NSURLConnection *)connection;
 
 - (TArray<uint8>&)getPayload;
+- (int32)getBytesWritten;
 @end
 
 
@@ -186,6 +191,11 @@ public:
 	 */
 	const int32 GetNumBytesReceived() const;
 
+	/**
+	* Get the number of bytes sent so far
+	*/
+	const int32 GetNumBytesWritten() const;
+
 	/**
 	 * Constructor
 	 *
diff --git a/Engine/Source/Runtime/Online/HTTP/Private/NullHttp.cpp b/Engine/Source/Runtime/Online/HTTP/Private/NullHttp.cpp
new file mode 100644
--- /dev/null
+++ b/Engine/Source/Runtime/Online/HTTP/Private/NullHttp.cpp
@@ -0,0 +1,201 @@
+// Copyright 1998-2015 Epic Games, Inc. All Rights Reserved.
+
+#include "HttpPrivatePCH.h"
+#include "NullHttp.h"
+
+// FNullHttpRequest
+
+FString FNullHttpRequest::GetURL()
+{
+	return Url;
+}
+
+FString FNullHttpRequest::GetURLParameter(const FString& ParameterName)
+{
+	return FString();
+}
+
+FString FNullHttpRequest::GetHeader(const FString& HeaderName)
+{
+	FString* Header = Headers.Find(HeaderName);
+	if (Header != NULL)
+	{
+		return *Header;
+	}
+	return FString();
+}
+
+TArray<FString> FNullHttpRequest::GetAllHeaders()
+{
+	TArray<FString> Result;
+	for (TMap<FString, FString>::TConstIterator It(Headers); It; ++It)
+	{
+		Result.Add(It.Key() + TEXT(": ") + It.Value());
+	}
+	return Result;
+}
+
+FString FNullHttpRequest::GetContentType()
+{
+	return GetHeader(TEXT("Content-Type"));
+}
+
+int32 FNullHttpRequest::GetContentLength()
+{
+	return Payload.Num();
+}
+
+const TArray<uint8>& FNullHttpRequest::GetContent()
+{
+	return Payload;
+}
+
+FString FNullHttpRequest::GetVerb()
+{
+	return Verb;
+}
+
+void FNullHttpRequest::SetVerb(const FString& InVerb)
+{
+	Verb = InVerb;
+}
+
+void FNullHttpRequest::SetURL(const FString& InURL)
+{
+	Url = InURL;
+}
+
+void FNullHttpRequest::SetContent(const TArray<uint8>& ContentPayload)
+{
+	Payload = ContentPayload;
+}
+
+void FNullHttpRequest::SetContentAsString(const FString& ContentString)
+{
+	FTCHARToUTF8 Converter(*ContentString);
+	Payload.SetNum(Converter.Length());
+	FMemory::Memcpy(Payload.GetData(), (uint8*)(ANSICHAR*)Converter.Get(), Payload.Num());
+}
+
+void FNullHttpRequest::SetHeader(const FString& HeaderName, const FString& HeaderValue)
+{
+	Headers.Add(HeaderName, HeaderValue);
+}
+
+bool FNullHttpRequest::ProcessRequest()
+{
+	ElapsedTime = 0;
+	CompletionStatus = EHttpRequestStatus::Processing;
+
+	UE_LOG(LogHttp, Log, TEXT("Start request. %p %s url=%s"), this, *GetVerb(), *GetURL());
+
+	FHttpModule::Get().GetHttpManager().AddRequest(SharedThis(this));
+	return true;
+}
+
+FHttpRequestCompleteDelegate& FNullHttpRequest::OnProcessRequestComplete()
+{
+	return RequestCompleteDelegate;
+}
+
+FHttpRequestProgressDelegate& FNullHttpRequest::OnRequestProgress()
+{
+	return ReuestProgressDelegate;
+}
+
+void FNullHttpRequest::CancelRequest()
+{
+	FinishedRequest();
+}
+
+EHttpRequestStatus::Type FNullHttpRequest::GetStatus()
+{
+	return CompletionStatus;
+}
+
+const FHttpResponsePtr FNullHttpRequest::GetResponse() const
+{
+	return FHttpResponsePtr(nullptr);
+}
+
+void FNullHttpRequest::Tick(float DeltaSeconds)
+{
+	if (CompletionStatus == EHttpRequestStatus::Processing)
+	{
+		ElapsedTime += DeltaSeconds;
+		const float HttpTimeout = FHttpModule::Get().GetHttpTimeout();
+		if (HttpTimeout > 0 && ElapsedTime >= HttpTimeout)
+		{
+			UE_LOG(LogHttp, Warning, TEXT("Timeout processing Http request. %p"),
+				this);
+
+			FinishedRequest();
+		}
+	}
+}
+
+float FNullHttpRequest::GetElapsedTime()
+{
+	return ElapsedTime;
+}
+
+void FNullHttpRequest::FinishedRequest()
+{
+	CompletionStatus = EHttpRequestStatus::Failed;
+	TSharedRef<IHttpRequest> Request = SharedThis(this);
+	FHttpModule::Get().GetHttpManager().RemoveRequest(Request);
+
+	UE_LOG(LogHttp, Log, TEXT("Finished request %p. no response %s url=%s elapsed=%.3f"),
+		this, *GetVerb(), *GetURL(), ElapsedTime);
+
+	OnProcessRequestComplete().ExecuteIfBound(Request, NULL, false);
+}
+
+// FNullHttpResponse
+
+FString FNullHttpResponse::GetURL()
+{
+	return FString();
+}
+
+FString FNullHttpResponse::GetURLParameter(const FString& ParameterName)
+{
+	return FString();
+}
+
+FString FNullHttpResponse::GetHeader(const FString& HeaderName)
+{
+	return FString();
+}
+
+TArray<FString> FNullHttpResponse::GetAllHeaders()
+{
+	return TArray<FString>();
+}
+
+FString FNullHttpResponse::GetContentType()
+{
+	return FString();
+}
+
+int32 FNullHttpResponse::GetContentLength()
+{
+	return 0;
+}
+
+const TArray<uint8>& FNullHttpResponse::GetContent()
+{
+	return Payload;
+}
+
+int32 FNullHttpResponse::GetResponseCode()
+{
+	return 0;
+}
+
+FString FNullHttpResponse::GetContentAsString()
+{
+	return FString();
+}
+
+
diff --git a/Engine/Source/Runtime/Online/HTTP/Private/NullHttp.h b/Engine/Source/Runtime/Online/HTTP/Private/NullHttp.h
new file mode 100644
--- /dev/null
+++ b/Engine/Source/Runtime/Online/HTTP/Private/NullHttp.h
@@ -0,0 +1,77 @@
+// Copyright 1998-2015 Epic Games, Inc. All Rights Reserved.
+
+#pragma once
+
+/**
+ * Null (mock) implementation of an HTTP request
+ */
+class FNullHttpRequest : public IHttpRequest
+{
+public:
+
+	// IHttpBase
+	virtual FString GetURL() override;
+	virtual FString GetURLParameter(const FString& ParameterName) override;
+	virtual FString GetHeader(const FString& HeaderName) override;
+	virtual TArray<FString> GetAllHeaders() override;	
+	virtual FString GetContentType() override;
+	virtual int32 GetContentLength() override;
+	virtual const TArray<uint8>& GetContent() override;
+	// IHttpRequest 
+	virtual FString GetVerb() override;
+	virtual void SetVerb(const FString& InVerb) override;
+	virtual void SetURL(const FString& InURL) override;
+	virtual void SetContent(const TArray<uint8>& ContentPayload) override;
+	virtual void SetContentAsString(const FString& ContentString) override;
+	virtual void SetHeader(const FString& HeaderName, const FString& HeaderValue) override;
+	virtual bool ProcessRequest() override;
+	virtual FHttpRequestCompleteDelegate& OnProcessRequestComplete() override;
+	virtual FHttpRequestProgressDelegate& OnRequestProgress() override;
+	virtual void CancelRequest() override;
+	virtual EHttpRequestStatus::Type GetStatus() override;
+	virtual const FHttpResponsePtr GetResponse() const override;
+	virtual void Tick(float DeltaSeconds) override;
+	virtual float GetElapsedTime() override;
+
+	FNullHttpRequest()
+		: CompletionStatus(EHttpRequestStatus::NotStarted)
+		, ElapsedTime(0)
+	{}
+	virtual ~FNullHttpRequest() {}
+
+private:
+	void FinishedRequest();
+
+	FString Url;
+	FString Verb;
+	TArray<uint8> Payload;
+	FHttpRequestCompleteDelegate RequestCompleteDelegate;
+	FHttpRequestProgressDelegate ReuestProgressDelegate;
+	EHttpRequestStatus::Type CompletionStatus;
+	TMap<FString, FString> Headers;
+	float ElapsedTime;
+};
+
+/**
+ * Null (mock) implementation of an HTTP request
+ */
+class FNullHttpResponse : public IHttpResponse
+{
+	// IHttpBase 
+	virtual FString GetURL() override;
+	virtual FString GetURLParameter(const FString& ParameterName) override;
+	virtual FString GetHeader(const FString& HeaderName) override;
+	virtual TArray<FString> GetAllHeaders() override;	
+	virtual FString GetContentType() override;
+	virtual int32 GetContentLength() override;
+	virtual const TArray<uint8>& GetContent() override;
+	// Begin IHttpResponse interface
+	virtual int32 GetResponseCode() override;
+	virtual FString GetContentAsString() override;
+
+	FNullHttpResponse() {}
+	virtual ~FNullHttpResponse() {}
+
+private:
+	TArray<uint8> Payload;
+};
