Add support to BuildPatchTool to output file sizes when enumerating data files.

diff --git a/Engine/Source/Programs/AutomationTool/MCPPublic.cs b/Engine/Source/Programs/AutomationTool/MCPPublic.cs
--- a/Engine/Source/Programs/AutomationTool/MCPPublic.cs
+++ b/Engine/Source/Programs/AutomationTool/MCPPublic.cs
@@ -341,6 +341,10 @@ namespace EpicGames.MCP.Automation
 			/// Matches the corresponding BuildPatchTool command line argument.
 			/// </summary>
 			public string OutputFile;
+			/// <summary>
+			/// When true, the output will include the size of individual files
+			/// </summary>
+			public bool bIncludeSize;
 		}
 
         static BuildPatchToolBase Handler = null;
diff --git a/Engine/Source/Programs/BuildPatchTool/Private/BuildPatchToolMain.cpp b/Engine/Source/Programs/BuildPatchTool/Private/BuildPatchToolMain.cpp
--- a/Engine/Source/Programs/BuildPatchTool/Private/BuildPatchToolMain.cpp
+++ b/Engine/Source/Programs/BuildPatchTool/Private/BuildPatchToolMain.cpp
@@ -75,6 +75,7 @@
 	Required arguments:
 	-ManifestFile=""	Specifies in quotes the file path to the manifest to enumerate from
 	-OutputFile=""		Specifies in quotes the file path to a file where the list will be saved out, \r\n separated cloud relative paths.
+	-includesizes		When specified, the size of each file will be output following the filename in the format " (xxx bytes)".
 	-dataenumerate		Must be specified to launch the tool in cloud seeder mode
 
 =============================================================================*/
@@ -183,6 +184,7 @@ int32 BuildPatchToolMain( const TCHAR* CommandLine )
 	bool bPreview = false;
 	bool bNoPatchDelete = false;
 	bool bPatchWithReuseAgeThreshold = true;
+	bool bIncludeSizes = false;
 
 	// Collect all the info from the CommandLine
 	TArray< FString > Tokens, Switches;
@@ -222,6 +224,7 @@ int32 BuildPatchToolMain( const TCHAR* CommandLine )
 		}
 		bPreview = Switches.IndexOfByPredicate(FCommandLineMatcher(TEXT("preview"))) != INDEX_NONE;
 		bNoPatchDelete = Switches.IndexOfByPredicate(FCommandLineMatcher(TEXT("nopatchdelete"))) != INDEX_NONE;
+		bIncludeSizes = Switches.IndexOfByPredicate(FCommandLineMatcher(TEXT("includesizes"))) != INDEX_NONE;
 		BuildRootIdx = Switches.IndexOfByPredicate(FCommandLineMatcher(TEXT("BuildRoot")));
 		CloudDirIdx = Switches.IndexOfByPredicate(FCommandLineMatcher(TEXT("CloudDir")));
 		AppIDIdx = Switches.IndexOfByPredicate(FCommandLineMatcher(TEXT("AppID")));
@@ -377,32 +380,32 @@ int32 BuildPatchToolMain( const TCHAR* CommandLine )
 			IniFile = CloudDirectory / TEXT("BuildPatchTool.ini");
 			GConfig->InitializeConfigSystem();
 
-		if (DataAgeThresholdIdx != INDEX_NONE)
-		{
-			FParse::Value(*Switches[DataAgeThresholdIdx], TEXT("DataAgeThreshold="), DataAgeThreshold);
-		}
+			if (DataAgeThresholdIdx != INDEX_NONE)
+			{
+				FParse::Value(*Switches[DataAgeThresholdIdx], TEXT("DataAgeThreshold="), DataAgeThreshold);
+			}
 			else
 			{
 				switch (ToolMode)
-		{
+				{
 				case EBuildPatchToolMode::Compactify:
-			// For compactification, if we don't pass in DataAgeThreshold, and it's not in BuildPatchTool.ini,
-			// then we set it to zero, to indicate that any unused chunks are valid for deletion
-			if (!GConfig->GetFloat(TEXT("Compactify"), TEXT("DataAgeThreshold"), DataAgeThreshold, IniFile))
-			{
-				GLog->Log(ELogVerbosity::Warning, TEXT("DataAgeThreshold not supplied, so all unreferenced data is eliglble for deletion. Note that this process is NOT compatible with any concurrently running patch generaiton processes"));
-				DataAgeThreshold = 0.0f;
-			}
+					// For compactification, if we don't pass in DataAgeThreshold, and it's not in BuildPatchTool.ini,
+					// then we set it to zero, to indicate that any unused chunks are valid for deletion
+					if (!GConfig->GetFloat(TEXT("Compactify"), TEXT("DataAgeThreshold"), DataAgeThreshold, IniFile))
+					{
+						GLog->Log(ELogVerbosity::Warning, TEXT("DataAgeThreshold not supplied, so all unreferenced data is eliglble for deletion. Note that this process is NOT compatible with any concurrently running patch generaiton processes"));
+						DataAgeThreshold = 0.0f;
+					}
 					break;
 				case EBuildPatchToolMode::PatchGeneration:
-			// For patch generation, if we don't pass in DataAgeThreshold, and it's not specified in BuildPatchTool.ini,
-			// then we set bChunkWithReuseAgeThreshold to false, which indicates that *all* patch data is valid for reuse
-			if (!GConfig->GetFloat(TEXT("PatchGeneration"), TEXT("DataAgeThreshold"), DataAgeThreshold, IniFile))
-			{
-				GLog->Log(ELogVerbosity::Warning, TEXT("DataAgeThreshold not supplied, so all existing data is eligible for reuse. Note that this process is NOT compatible with any concurrently running compactify processes"));
-				DataAgeThreshold = 0.0f;
-				bPatchWithReuseAgeThreshold = false;
-			}
+					// For patch generation, if we don't pass in DataAgeThreshold, and it's not specified in BuildPatchTool.ini,
+					// then we set bChunkWithReuseAgeThreshold to false, which indicates that *all* patch data is valid for reuse
+					if (!GConfig->GetFloat(TEXT("PatchGeneration"), TEXT("DataAgeThreshold"), DataAgeThreshold, IniFile))
+					{
+						GLog->Log(ELogVerbosity::Warning, TEXT("DataAgeThreshold not supplied, so all existing data is eligible for reuse. Note that this process is NOT compatible with any concurrently running compactify processes"));
+						DataAgeThreshold = 0.0f;
+						bPatchWithReuseAgeThreshold = false;
+					}
 					break;
 				}
 			}
@@ -451,7 +454,7 @@ int32 BuildPatchToolMain( const TCHAR* CommandLine )
 	// Run the mode!
 	switch (ToolMode)
 	{
-		case EBuildPatchToolMode::Compactify:
+	case EBuildPatchToolMode::Compactify:
 	{
 		// Split out our manifests to keep arg (if any) into an array of manifest filenames
 		TArray<FString> ManifestsArr;
@@ -492,7 +495,7 @@ int32 BuildPatchToolMain( const TCHAR* CommandLine )
 		bSuccess = BuildPatchServicesModule->CompactifyCloudDirectory(ManifestsArr, DataAgeThreshold, CompactifyMode);
 	}
 		break;
-		case EBuildPatchToolMode::PatchGeneration:
+	case EBuildPatchToolMode::PatchGeneration:
 	{
 		FBuildPatchSettings Settings;
 		Settings.RootDirectory = RootDirectory + TEXT("/");
@@ -521,18 +524,18 @@ int32 BuildPatchToolMain( const TCHAR* CommandLine )
 		}
 	}
 		break;
-		case EBuildPatchToolMode::DataEnumeration:
-		{
-			// Run the data enumeration routine
-			bSuccess = BuildPatchServicesModule->EnumerateManifestData(MoveTemp(ManifestFile), MoveTemp(OutputFile));
-		}
+	case EBuildPatchToolMode::DataEnumeration:
+	{
+		// Run the data enumeration routine
+		bSuccess = BuildPatchServicesModule->EnumerateManifestData(MoveTemp(ManifestFile), MoveTemp(OutputFile), bIncludeSizes);
+	}
 		break;
-		default:
+	default:
 	{
 		GLog->Log(ELogVerbosity::Error, TEXT("Unknown tool mode"));
 		BuildPatchServicesModule.Reset();
 		FCoreDelegates::OnExit.Broadcast();
-		}
+	}
 		return 3;
 	}
 
diff --git a/Engine/Source/Runtime/Online/BuildPatchServices/Private/BuildPatchDataEnumeration.cpp b/Engine/Source/Runtime/Online/BuildPatchServices/Private/BuildPatchDataEnumeration.cpp
--- a/Engine/Source/Runtime/Online/BuildPatchServices/Private/BuildPatchDataEnumeration.cpp
+++ b/Engine/Source/Runtime/Online/BuildPatchServices/Private/BuildPatchDataEnumeration.cpp
@@ -4,7 +4,7 @@
 
 #if WITH_BUILDPATCHGENERATION
 
-bool FBuildDataEnumeration::EnumerateManifestData(FString ManifestFilePath, FString OutputFile)
+bool FBuildDataEnumeration::EnumerateManifestData(FString ManifestFilePath, FString OutputFile, const bool bIncludeSizes)
 {
 	bool bSuccess = false;
 	TArray<FGuid> DataList;
@@ -26,6 +26,11 @@ bool FBuildDataEnumeration::EnumerateManifestData(FString ManifestFilePath, FStr
 		{
 			FString Path = FBuildPatchUtils::GetDataFilename(AppManifest, FString(), DataGuid);
 			FullList += Path;
+			if (bIncludeSizes)
+			{
+				uint64 FileSize = AppManifest->GetDataSize(DataGuid);
+				FullList += FString::Printf(TEXT(" (%u bytes)"), FileSize);
+			}
 			FullList += TEXT("\r\n");
 			GLog->Log(ELogVerbosity::Log, Path);
 		}
diff --git a/Engine/Source/Runtime/Online/BuildPatchServices/Private/BuildPatchDataEnumeration.h b/Engine/Source/Runtime/Online/BuildPatchServices/Private/BuildPatchDataEnumeration.h
--- a/Engine/Source/Runtime/Online/BuildPatchServices/Private/BuildPatchDataEnumeration.h
+++ b/Engine/Source/Runtime/Online/BuildPatchServices/Private/BuildPatchDataEnumeration.h
@@ -7,7 +7,7 @@
 class FBuildDataEnumeration
 {
 public:
-	static bool EnumerateManifestData(FString ManifestFilePath, FString OutputFile);
+	static bool EnumerateManifestData(FString ManifestFilePath, FString OutputFile, const bool bIncludeSizes);
 };
 
 #endif //WITH_BUILDPATCHGENERATION
diff --git a/Engine/Source/Runtime/Online/BuildPatchServices/Private/BuildPatchServicesModule.cpp b/Engine/Source/Runtime/Online/BuildPatchServices/Private/BuildPatchServicesModule.cpp
--- a/Engine/Source/Runtime/Online/BuildPatchServices/Private/BuildPatchServicesModule.cpp
+++ b/Engine/Source/Runtime/Online/BuildPatchServices/Private/BuildPatchServicesModule.cpp
@@ -197,9 +197,9 @@ bool FBuildPatchServicesModule::CompactifyCloudDirectory(const TArray<FString>&
 	return FBuildDataCompactifier::CompactifyCloudDirectory(ManifestsToKeep, DataAgeThreshold, bPreview, bNoPatchDelete);
 }
 
-bool FBuildPatchServicesModule::EnumerateManifestData(FString ManifestFilePath, FString OutputFile)
+bool FBuildPatchServicesModule::EnumerateManifestData(FString ManifestFilePath, FString OutputFile, const bool bIncludeSizes)
 {
-	return FBuildDataEnumeration::EnumerateManifestData(MoveTemp(ManifestFilePath), MoveTemp(OutputFile));
+	return FBuildDataEnumeration::EnumerateManifestData(MoveTemp(ManifestFilePath), MoveTemp(OutputFile), bIncludeSizes);
 }
 
 #endif //WITH_BUILDPATCHGENERATION
diff --git a/Engine/Source/Runtime/Online/BuildPatchServices/Private/BuildPatchServicesModule.h b/Engine/Source/Runtime/Online/BuildPatchServices/Private/BuildPatchServicesModule.h
--- a/Engine/Source/Runtime/Online/BuildPatchServices/Private/BuildPatchServicesModule.h
+++ b/Engine/Source/Runtime/Online/BuildPatchServices/Private/BuildPatchServicesModule.h
@@ -86,7 +86,7 @@ public:
 	virtual bool GenerateChunksManifestFromDirectory( const FBuildPatchSettings& Settings ) override;
 	virtual bool GenerateFilesManifestFromDirectory( const FBuildPatchSettings& Settings ) override;
 	virtual bool CompactifyCloudDirectory( const TArray<FString>& ManifestsToKeep, const float DataAgeThreshold, const ECompactifyMode::Type Mode ) override;
-	virtual bool EnumerateManifestData(FString ManifestFilePath, FString OutputFile) override;
+	virtual bool EnumerateManifestData(FString ManifestFilePath, FString OutputFile, const bool bIncludeSizes) override;
 #endif
 	virtual IBuildManifestPtr MakeManifestFromJSON( const FString& ManifestJSON ) override;
 	// End IBuildPatchServicesModule interface
diff --git a/Engine/Source/Runtime/Online/BuildPatchServices/Public/Interfaces/IBuildPatchServicesModule.h b/Engine/Source/Runtime/Online/BuildPatchServices/Public/Interfaces/IBuildPatchServicesModule.h
--- a/Engine/Source/Runtime/Online/BuildPatchServices/Public/Interfaces/IBuildPatchServicesModule.h
+++ b/Engine/Source/Runtime/Online/BuildPatchServices/Public/Interfaces/IBuildPatchServicesModule.h
@@ -141,9 +141,10 @@ public:
 	 * Saves to file, and logs, a full list of cloud dir relative referenced data file paths
 	 * @param   ManifestFilePath     A full file path for the manifest to be loaded
 	 * @param   OutputFile           A full file path where to save the output text
+	 * @param   bIncludeSizes        If true, will include the size (in bytes) with every file output
 	 * @return  true if successful
 	 */
-	virtual bool EnumerateManifestData(FString ManifestFilePath, FString OutputFile) = 0;
+	virtual bool EnumerateManifestData(FString ManifestFilePath, FString OutputFile, const bool bIncludeSizes) = 0;
 
 #endif // WITH_BUILDPATCHGENERATION
 
