Fix for BuildPatchTool's threads crashing due to FRunnableThread::InitializeTls() call being required by all programs that use runnable threads.

diff --git a/Engine/Source/Programs/BuildPatchTool/Private/BuildPatchToolMain.cpp b/Engine/Source/Programs/BuildPatchTool/Private/BuildPatchToolMain.cpp
--- a/Engine/Source/Programs/BuildPatchTool/Private/BuildPatchToolMain.cpp
+++ b/Engine/Source/Programs/BuildPatchTool/Private/BuildPatchToolMain.cpp
@@ -158,6 +158,10 @@ int32 BuildPatchToolMain( const TCHAR* CommandLine )
 	GLog->Logf(TEXT("BuildPatchToolMain ran with: %s"), CommandLine);
 
 	FPlatformProcess::SetCurrentWorkingDirectoryToBaseDir();
+
+	// Init TLS for runnable threads.
+	FRunnableThread::InitializeTls();
+
 	bool bSuccess = false;
 
 	FString RootDirectory;
