Ensure BuildPatchTool does a full engine preinit.

diff --git a/Engine/Source/Programs/BuildPatchTool/Private/BuildPatchToolMain.cpp b/Engine/Source/Programs/BuildPatchTool/Private/BuildPatchToolMain.cpp
--- a/Engine/Source/Programs/BuildPatchTool/Private/BuildPatchToolMain.cpp
+++ b/Engine/Source/Programs/BuildPatchTool/Private/BuildPatchToolMain.cpp
@@ -142,23 +142,16 @@ namespace EBuildPatchToolMode
 
 int32 BuildPatchToolMain( const TCHAR* CommandLine )
 {
-	// Initialize the command line
-	FCommandLine::Set(CommandLine);
+	GEngineLoop.PreInit(CommandLine);
 
-	// Add log devices
+	// Add log device
 	if (FParse::Param(FCommandLine::Get(), TEXT("stdout")))
 	{
 		GLog->AddOutputDevice(new FBuildPatchOutputDevice());
 	}
-	if (FPlatformMisc::IsDebuggerPresent())
-	{
-		GLog->AddOutputDevice(new FOutputDeviceDebug());
-	}
 
 	GLog->Logf(TEXT("BuildPatchToolMain ran with: %s"), CommandLine);
 
-	FPlatformProcess::SetCurrentWorkingDirectoryToBaseDir();
-
 	bool bSuccess = false;
 
 	FString RootDirectory;
