Fixed BuildPatchTool warnings caused by missing module dependencies

diff --git a/Engine/Source/Programs/BuildPatchTool/BuildPatchTool.Build.cs b/Engine/Source/Programs/BuildPatchTool/BuildPatchTool.Build.cs
--- a/Engine/Source/Programs/BuildPatchTool/BuildPatchTool.Build.cs
+++ b/Engine/Source/Programs/BuildPatchTool/BuildPatchTool.Build.cs
@@ -17,7 +17,12 @@ public class BuildPatchTool : ModuleRules
 				"Core",
 				"BuildPatchServices",
 				"Projects",
+				// The below items are not strictly needed by BPT, but core appears to need them during initialization
+				"PakFile",
+				"SandboxFile",
+				"NetworkFile",
+				"StreamingFile"
 			}
 		);
 	}
-}
\ No newline at end of file
+}
