Merging using UE4-To-UE4-LauncherDev

CL# 2644628 on 05/08/2015 05:12
---------------------
Quick comment fix for BuildPatchTool main cpp

diff --git a/Engine/Source/Programs/BuildPatchTool/Private/BuildPatchToolMain.cpp b/Engine/Source/Programs/BuildPatchTool/Private/BuildPatchToolMain.cpp
--- a/Engine/Source/Programs/BuildPatchTool/Private/BuildPatchToolMain.cpp
+++ b/Engine/Source/Programs/BuildPatchTool/Private/BuildPatchToolMain.cpp
@@ -6,6 +6,7 @@
 	The tool can be used in two modes of operation:
 	(1) to generate patch data (manifest, chunks, files) for build images given the existing cloud data; or
 	(2) to "compactify" a cloud directory by removing all orphaned chunks, not referenced by any manifest file.
+	(3) to enumerate the data files referenced by a manifest
 
 	In order to trigger compactify functionality, the -compactify commandline argument should be specified.
 
