[INTEGRATE] Change 2702276 by [REDACTED_PERSON][REDACTED_HANDLE] on 2015/09/23 05:17

        Alter BuildPatchTool main function to return values from an enum to aid in managing return codes.

diff --git a/Engine/Source/Programs/BuildPatchTool/Private/BuildPatchToolMain.cpp b/Engine/Source/Programs/BuildPatchTool/Private/BuildPatchToolMain.cpp
--- a/Engine/Source/Programs/BuildPatchTool/Private/BuildPatchToolMain.cpp
+++ b/Engine/Source/Programs/BuildPatchTool/Private/BuildPatchToolMain.cpp
@@ -139,8 +139,22 @@ namespace EBuildPatchToolMode
 	};
 }
 
+namespace EBuildPatchToolReturnCode
+{
+	enum Type
+	{
+		OK = 0,
+		ArgumentProcessingError,
+		CouldNotOpenManifestToKeepFile,
+		UnknownToolMode,
+		UnknownError,
+		LegacyDoNotUse,
+		FileIgnoreListNotFound,
+		FileAttributeListNotFound
+	};
+}
 
-int32 BuildPatchToolMain( const TCHAR* CommandLine )
+EBuildPatchToolReturnCode::Type BuildPatchToolMain( const TCHAR* CommandLine )
 {
 	// Add log device
 	if (FParse::Param(CommandLine, TEXT("stdout")))
@@ -409,19 +423,19 @@ int32 BuildPatchToolMain( const TCHAR* CommandLine )
 	if( !bSuccess )
 	{
 		GLog->Log(ELogVerbosity::Error, TEXT("An error occurred processing arguments"));
-		return 1;
+		return EBuildPatchToolReturnCode::ArgumentProcessingError;
 	}
 
 	if (!IgnoreListFile.IsEmpty() && !FPaths::FileExists(IgnoreListFile))
 	{
 		GLog->Logf(ELogVerbosity::Error, TEXT("Provided file ignore list was not found %s"), *IgnoreListFile);
-		return 6;
+		return EBuildPatchToolReturnCode::FileIgnoreListNotFound;
 	}
 
 	if (!AttributeListFile.IsEmpty() && !FPaths::FileExists(AttributeListFile))
 	{
 		GLog->Logf(ELogVerbosity::Error, TEXT("Provided file attribute list was not found %s"), *AttributeListFile);
-		return 7;
+		return EBuildPatchToolReturnCode::FileAttributeListNotFound;
 	}
 
 	// Load the BuildPatchServices Module
@@ -460,7 +474,7 @@ int32 BuildPatchToolMain( const TCHAR* CommandLine )
 				GLog->Log(ELogVerbosity::Error, TEXT("Could not open specified manifests to keep file"));
 				BuildPatchServicesModule.Reset();
 				FCoreDelegates::OnExit.Broadcast();
-				return 2;
+				return EBuildPatchToolReturnCode::CouldNotOpenManifestToKeepFile;
 			}
 		}
 
@@ -516,7 +530,7 @@ int32 BuildPatchToolMain( const TCHAR* CommandLine )
 		BuildPatchServicesModule.Reset();
 		FCoreDelegates::OnExit.Broadcast();
 	}
-		return 3;
+		return EBuildPatchToolReturnCode::UnknownToolMode;
 	}
 
 	// Release the module ptr
@@ -525,15 +539,15 @@ int32 BuildPatchToolMain( const TCHAR* CommandLine )
 	// Check for processing error
 	if (!bSuccess)
 	{
-		GLog->Log(ELogVerbosity::Error, TEXT("A fatal error occurred executing BuildPatchTool.exe"));
+		GLog->Log(ELogVerbosity::Error, TEXT("An unknown error occurred executing BuildPatchTool.exe"));
 		FCoreDelegates::OnExit.Broadcast();
-		return 4;
+		return EBuildPatchToolReturnCode::UnknownError;
 	}
 
 	FCoreDelegates::OnExit.Broadcast();
 
 	GLog->Log(TEXT("BuildPatchToolMain completed successfuly"));
-	return 0;
+	return EBuildPatchToolReturnCode::OK;
 }
 
 INT32_MAIN_INT32_ARGC_TCHAR_ARGV()
@@ -560,5 +574,5 @@ INT32_MAIN_INT32_ARGC_TCHAR_ARGV()
 		CommandLine += Argument;
 	}
 
-	return BuildPatchToolMain( *CommandLine );
+	return static_cast<int32>(BuildPatchToolMain( *CommandLine ));
 }
