Fixing build warning for Build Patch Tool
warning C4996: 'FOutputDevice::FormatLogLine': Please use FOutputDeviceHelper::FormatLogLine. Please update your code to the new API before upgrading to the next release, otherwise your project will no longer compile.

diff --git a/Engine/Source/Programs/BuildPatchTool/Private/BuildPatchTool.cpp b/Engine/Source/Programs/BuildPatchTool/Private/BuildPatchTool.cpp
--- a/Engine/Source/Programs/BuildPatchTool/Private/BuildPatchTool.cpp
+++ b/Engine/Source/Programs/BuildPatchTool/Private/BuildPatchTool.cpp
@@ -18,9 +18,9 @@ public:
 	{
 #if PLATFORM_WINDOWS
 #if PLATFORM_USE_LS_SPEC_FOR_WIDECHAR
-		printf( "\n%ls", *FOutputDevice::FormatLogLine( Verbosity, Category, V, GPrintLogTimes ) );
+		printf( "\n%ls", *FOutputDeviceHelper::FormatLogLine( Verbosity, Category, V, GPrintLogTimes ) );
 #else
-		wprintf( TEXT( "\n%s" ), *FOutputDevice::FormatLogLine( Verbosity, Category, V, GPrintLogTimes ) );
+		wprintf( TEXT( "\n%s" ), *FOutputDeviceHelper::FormatLogLine( Verbosity, Category, V, GPrintLogTimes ) );
 #endif
 		fflush( stdout );
 #endif
