Merging //UE4/Portal-Staging/ to //UE4/Release-4.23/

CL# 6606168 by [REDACTED_PERSON] on 2019/05/22
---------------
Fix for BPT shutdown crash. Some engine code crashes if GIsRequestingExit is not true when exiting.

diff --git a/Engine/Source/Programs/BuildPatchTool/Private/BuildPatchTool.cpp b/Engine/Source/Programs/BuildPatchTool/Private/BuildPatchTool.cpp
--- a/Engine/Source/Programs/BuildPatchTool/Private/BuildPatchTool.cpp
+++ b/Engine/Source/Programs/BuildPatchTool/Private/BuildPatchTool.cpp
@@ -142,6 +142,7 @@ EReturnCode BuildPatchToolMain(const TCHAR* CommandLine)
 	}
 
 	// Shutdown
+	GIsRequestingExit = true;
 	FCoreDelegates::OnExit.Broadcast();
 
 	return ReturnCode;
