Fix BPT exception when generating project files.

#rnx

diff --git a/Engine/Source/Programs/BuildPatchTool/BuildPatchTool.Build.cs b/Engine/Source/Programs/BuildPatchTool/BuildPatchTool.Build.cs
--- a/Engine/Source/Programs/BuildPatchTool/BuildPatchTool.Build.cs
+++ b/Engine/Source/Programs/BuildPatchTool/BuildPatchTool.Build.cs
@@ -32,6 +32,7 @@ public class BuildPatchTool : ModuleRules
 		{
 			PrivateDependencyModuleNames.AddRange(
 				new string[] {
+					"ApplicationCore", // Required by AutomationController
 					"Messaging",
 					"AutomationWorker",
 					"AutomationController"
diff --git a/Engine/Source/Programs/BuildPatchTool/BuildPatchTool.Target.cs b/Engine/Source/Programs/BuildPatchTool/BuildPatchTool.Target.cs
--- a/Engine/Source/Programs/BuildPatchTool/BuildPatchTool.Target.cs
+++ b/Engine/Source/Programs/BuildPatchTool/BuildPatchTool.Target.cs
@@ -18,7 +18,11 @@ public class BuildPatchToolTarget : TargetRules
 		bBuildEditor = false;
 		bCompileAgainstEngine = false;
 		bCompileAgainstCoreUObject = false;
-		bCompileAgainstApplicationCore = false;
+		if(Target.Configuration == UnrealTargetConfiguration.Shipping)
+		{
+			// AutomationController is referenced by BuildPatchTool.build.cs, and references a ton of editor things that it probably shouldn't. Need ApplicationCore for this.
+			bCompileAgainstApplicationCore = false;
+		}
 		bCompileLeanAndMeanUE = true;
 		bUseLoggingInShipping = true;
 		bUseChecksInShipping = true;
