Properly call AppExit when BuildPatchTool is finished to avoid shutdown order crashes
#rnx

diff --git a/Engine/Source/Programs/BuildPatchTool/Private/BuildPatchTool.cpp b/Engine/Source/Programs/BuildPatchTool/Private/BuildPatchTool.cpp
--- a/Engine/Source/Programs/BuildPatchTool/Private/BuildPatchTool.cpp
+++ b/Engine/Source/Programs/BuildPatchTool/Private/BuildPatchTool.cpp
@@ -127,7 +127,9 @@ EReturnCode BuildPatchToolMain(const TCHAR* CommandLine)
 
 	// Shutdown
 	RequestEngineExit(TEXT("BuildPatchToolMain Exiting"));
-	FCoreDelegates::OnExit.Broadcast();
+
+	FEngineLoop::AppPreExit();
+	FEngineLoop::AppExit();
 
 	return ReturnCode;
 }
