Fix deprecation warning, use AutomationControllerManager::GetEnabledReports instead of GetReports

diff --git a/Engine/Source/Programs/BuildPatchTool/Private/ToolModes/AutomationMode.cpp b/Engine/Source/Programs/BuildPatchTool/Private/ToolModes/AutomationMode.cpp
--- a/Engine/Source/Programs/BuildPatchTool/Private/ToolModes/AutomationMode.cpp
+++ b/Engine/Source/Programs/BuildPatchTool/Private/ToolModes/AutomationMode.cpp
@@ -100,7 +100,8 @@ public:
 		}
 
 		// Check for failures and exit.
-		bool bSuccess = !GIsCriticalError && RecursiveCheckReports(AutomationControllerManager->GetReports());
+		TArray<TSharedPtr<IAutomationReport>> Reports = AutomationControllerManager->GetEnabledReports();
+		bool bSuccess = !GIsCriticalError && RecursiveCheckReports(Reports);
 		return bSuccess ? EReturnCode::OK : EReturnCode::ToolFailure;
 	}
 
