CIS fix.

diff --git a/Engine/Source/Programs/BuildPatchTool/Public/Interfaces/ToolMode.h b/Engine/Source/Programs/BuildPatchTool/Public/Interfaces/ToolMode.h
--- a/Engine/Source/Programs/BuildPatchTool/Public/Interfaces/ToolMode.h
+++ b/Engine/Source/Programs/BuildPatchTool/Public/Interfaces/ToolMode.h
@@ -45,7 +45,7 @@ namespace BuildPatchTool
 			{
 				// Overwrite trailing slash with terminator, and trim size.
 				UriPath.GetCharArray()[TrailingSlashIdx--] = 0;
-				UriPath.GetCharArray().SetNum(UriPath.GetCharArray().Num() - 1, false);
+				UriPath.GetCharArray().SetNum(UriPath.GetCharArray().Num() - 1, EAllowShrinking::No);
 			}
 		}
 
