BPT: Fix for crash on exit unloading modules.

#rnx

diff --git a/Engine/Source/Programs/BuildPatchTool/Private/BuildPatchTool.cpp b/Engine/Source/Programs/BuildPatchTool/Private/BuildPatchTool.cpp
--- a/Engine/Source/Programs/BuildPatchTool/Private/BuildPatchTool.cpp
+++ b/Engine/Source/Programs/BuildPatchTool/Private/BuildPatchTool.cpp
@@ -129,6 +129,7 @@ EReturnCode BuildPatchToolMain(const TCHAR* CommandLine)
 	RequestEngineExit(TEXT("BuildPatchToolMain Exiting"));
 
 	FEngineLoop::AppPreExit();
+	FModuleManager::Get().UnloadModulesAtShutdown();
 	FEngineLoop::AppExit();
 
 	return ReturnCode;
