p4 integrate //UE5/Dev-Distro-5.5/Engine/Source/Programs/BuildPatchTool/... //UE5/Main/Engine/Source/Programs/BuildPatchTool/...
p4 integrate //UE5/Dev-Distro-5.5/Engine/Programs/BuildPatchTool/... //UE5/Main/Engine/Programs/BuildPatchTool/...
p4 integrate //UE5/Dev-Distro-5.5/Engine/Source/Runtime/Online/BuildPatchServices/... //UE5/Main/Engine/Source/Runtime/Online/BuildPatchServices/...

#Synced-CL 50932611

diff --git a/Engine/Programs/BuildPatchTool/Config/DefaultEngine.ini b/Engine/Programs/BuildPatchTool/Config/DefaultEngine.ini
--- a/Engine/Programs/BuildPatchTool/Config/DefaultEngine.ini
+++ b/Engine/Programs/BuildPatchTool/Config/DefaultEngine.ini
@@ -10,3 +10,114 @@ LogTextLocalizationManager=warning
 LogObj=warning
 LogExit=warning
 LogPlatformFile=warning
+LogICUInternationalization=warning
+LogHttp=error
+LogOnline=error
+
+[BuildPatchTool.ConsoleLog]
++ExcludedCategories=LogOnline
++ExcludedCategories=LogHttp
+
+[OnlineSubsystem]
+DefaultPlatformService=Mcp
+
+[OnlineSubsystemMcp]
+bEnabled=true
+
+[HTTP]
+HttpMaxConnectionsPerServer=0
+
+; OnlinePartyMcp is not needed, but spews warnings if you do not provide default config
+[OnlineSubsystemMcp.OnlinePartyMcp]
+PartyPendingCreateResponseTimeout=0
+PartyPendingJoinResponseTimeout=20
+PartyPendingLeaveResponseTimeout=10
+PartyPendingKickResponseTimeout=10
+PartyPendingJoinRequestTimeout=20
+
+;========== ACCOUNT SERVICE ==========
+[OnlineSubsystemMcp.OnlineIdentityMcp]
+ServiceName=account
+Protocol=https
+
+[OnlineSubsystemMcp.OnlineIdentityMcp Prod]
+Domain=account-public-service-prod03.ol.epicgames.com
+RedirectUrl=accounts.epicgames.com
+
+[OnlineSubsystemMcp.OnlineIdentityMcp Stage]
+Domain=account-public-service-stage.ol.epicgames.com
+RedirectUrl=accounts.epicgames-stage.ol.epicgames.com
+
+[OnlineSubsystemMcp.OnlineIdentityMcp GameDev]
+Domain=account-public-service-common-gamedev.ol.epicgames.net
+RedirectUrl=accounts.epicgames-gamedev.ol.epicgames.net
+
+[OnlineSubsystemMcp.OnlineIdentityMcp LocalHost]
+Domain=account-public-service-common-gamedev.ol.epicgames.net
+RedirectUrl=accounts.epicgames-gamedev.ol.epicgames.net
+
+[OnlineSubsystemMcp.OnlineIdentityMcp CI]
+Domain=account-public-service-common-ci.ol.epicgames.net
+RedirectUrl=accounts.epicgames-ci.ol.epicgames.net
+
+
+;========== ARTIFACT SERVICE ==========
+[OnlineSubsystemMcp.OnlineArtifactServiceMcp]
+ServiceName=artifact-service
+Protocol=https
+
+[OnlineSubsystemMcp.OnlineArtifactServiceMcp Prod]
+Domain=artifact-public-service-prod.ol.epicgames.com
+
+[OnlineSubsystemMcp.OnlineArtifactServiceMcp Stage]
+Domain=artifact-public-service-stage.ol.epicgames.com
+
+[OnlineSubsystemMcp.OnlineArtifactServiceMcp GameDev]
+Domain=artifact-public-service-gamedev.ol.epicgames.net
+
+[OnlineSubsystemMcp.OnlineArtifactServiceMcp CI]
+Domain=artifact-public-service-ci.ol.epicgames.net
+
+[OnlineSubsystemMcp.OnlineArtifactServiceMcp LocalHost]
+Domain=artifact-public-service-gamedev.ol.epicgames.net
+
+;========== ARTIFACT DELIVERY SERVICE ==========
+[OnlineSubsystemMcp.OnlineArtifactDeliveryServiceMcp]
+ServiceName=artifact-delivery
+Protocol=https
+
+[OnlineSubsystemMcp.OnlineArtifactDeliveryServiceMcp GameDev]
+Domain=artifact-delivery-public-service-gamedev.aeee.dev.use1a.on.epicgames.com
+
+[OnlineSubsystemMcp.OnlineArtifactDeliveryServiceMcp Prod]
+Domain=artifact-delivery-service-public-prod.ol.epicgames.com
++AltDomains=artifact-delivery-public-service-prod.beee.live.use1a.on.epicgames.com
+
+;========== LAUNCHER SERVICE ==========
+[OnlineSubsystemMcp.OnlineBuildInfoMcpV2]
+ServiceName=launcher
+Protocol=https
+
+[OnlineSubsystemMcp.OnlineBuildInfoMcpV2 Prod]
+Domain=launcher-public-service-prod06.ol.epicgames.com
+
+[OnlineSubsystemMcp.OnlineBuildInfoMcpV2 Stage]
+Domain=launcher-public-service-stage.ol.epicgames.
+
+[OnlineSubsystemMcp.OnlineBuildInfoMcpV2 GameDev]
+Domain=launcher-public-service-gamedev.ol.epicgames.net
+
+[OnlineSubsystemMcp.OnlineBuildInfoMcpV2 CI]
+Domain=launcher-public-service-ci.ol.epicgames.net
+
+[OnlineSubsystemMcp.OnlineBuildInfoMcpV2 LocalHost]
+Domain=launcher-public-service-gamedev.ol.epicgames.net
+
+;========== CONTENT PROTECTION SERVICE ==========
+
+[OnlineSubsystemMcp.ContentProtectionServiceMcp]
+ServiceName=contentprotection
+Protocol=https
+
+[OnlineSubsystemMcp.ContentProtectionServiceMcp GameDev]
+Domain=content-protection-service-gamedev.dev.use1a.on.epicgames.net
diff --git a/Engine/Source/Programs/BuildPatchTool/BuildPatchTool.Build.cs b/Engine/Source/Programs/BuildPatchTool/BuildPatchTool.Build.cs
--- a/Engine/Source/Programs/BuildPatchTool/BuildPatchTool.Build.cs
+++ b/Engine/Source/Programs/BuildPatchTool/BuildPatchTool.Build.cs
@@ -1,12 +1,13 @@
 // Copyright Epic Games, Inc. All Rights Reserved.
 
 using UnrealBuildTool;
+using System.IO;
 
 public class BuildPatchTool : ModuleRules
 {
-	public BuildPatchTool( ReadOnlyTargetRules Target ) : base(Target)
+	public BuildPatchTool(ReadOnlyTargetRules Target) : base(Target)
 	{
-		PublicIncludePathModuleNames.Add("Launch");
+		PublicIncludePaths.Add("Runtime/Launch/Public");
 
 		PrivateDependencyModuleNames.AddRange(
 			new string[]
@@ -15,26 +16,14 @@ public class BuildPatchTool : ModuleRules
 				"CoreUObject",
 				"BuildPatchServices",
 				"Projects",
+				"Networking",
 				"HTTP",
-				// The below items are not strictly needed by BPT, but core appears to need them during initialization
-				"PakFile",
-				"SandboxFile",
-				"NetworkFile",
-				"StreamingFile"
+				"PlatformCryptoOpenSSL",
 			}
 		);
-
-
-		if (Target.Configuration != UnrealTargetConfiguration.Shipping)
-		{
-			PrivateDependencyModuleNames.AddRange(
-				new string[] {
-					"ApplicationCore", // Required by AutomationController
-					"Messaging",
-					"AutomationWorker",
-					"AutomationController"
-				}
-			);
-		}
+		
+		PublicDefinitions.Add("WITH_GOOGLE_MOCK=0");
+		PublicDefinitions.Add("WITH_GOOGLE_TEST=0");
+		PublicDefinitions.Add("WITH_ONLINEMCP=0");
 	}
 }
diff --git a/Engine/Source/Programs/BuildPatchTool/BuildPatchTool.Target.cs b/Engine/Source/Programs/BuildPatchTool/BuildPatchTool.Target.cs
--- a/Engine/Source/Programs/BuildPatchTool/BuildPatchTool.Target.cs
+++ b/Engine/Source/Programs/BuildPatchTool/BuildPatchTool.Target.cs
@@ -10,24 +10,36 @@ public class BuildPatchToolTarget : TargetRules
 	public BuildPatchToolTarget(TargetInfo Target) : base(Target)
 	{
 		Type = TargetType.Program;
-		LinkType = TargetLinkType.Monolithic;
 		IncludeOrderVersion = EngineIncludeOrderVersion.Latest;
+		DefaultBuildSettings = BuildSettingsVersion.Latest;
+		LinkType = TargetLinkType.Monolithic;
+
 		LaunchModuleName = "BuildPatchTool";
-        bLegalToDistributeBinary = true;
+		bLegalToDistributeBinary = true;
 		UndecoratedConfiguration = UnrealTargetConfiguration.Shipping;
 
 		bCompileAgainstEngine = false;
 		bCompileAgainstCoreUObject = false;
-		if(Target.Configuration == UnrealTargetConfiguration.Shipping)
-		{
-			// AutomationController is referenced by BuildPatchTool.build.cs, and references a ton of editor things that it probably shouldn't. Need ApplicationCore for this.
-			bCompileAgainstApplicationCore = false;
-		}
+		bCompileAgainstApplicationCore = false;
 		bBuildDeveloperTools = false;
 		bUseLoggingInShipping = true;
 		bUseChecksInShipping = true;
 		bIsBuildingConsoleApplication = true;
 		bHasExports = false;
 		bWithServerCode = false;
+		bForceDisableAutomationTests = true;
+
+		GlobalDefinitions.Add("DISABLE_CWD_CHANGES=1");
+		GlobalDefinitions.Add("UE_DONT_USE_DEFAULT_OUTPUT_DEVICES=1");
+		
+		// We don't use Crash Reporter right now.
+		GlobalDefinitions.Add("NOINITCRASHREPORTER=1");
+
+		// This is needed to avoid adding a dependency on thirdparty dynamic libraries that BPT doesn't use
+		GlobalDefinitions.Add("DISABLE_AUDIO_MCP_MODULES=1");
+
+		// Add definitions required by our windows RC file.
+		GlobalDefinitions.Add(string.Format("PORTAL_CHANGELIST={0}", Target.Version.Changelist));
+		GlobalDefinitions.Add(string.Format("PORTAL_BRANCH_NAME=\"{0}\"", Target.Version.BranchName));
 	}
 }
diff --git a/Engine/Source/Programs/BuildPatchTool/Private/BuildPatchTool.cpp b/Engine/Source/Programs/BuildPatchTool/Private/BuildPatchTool.cpp
--- a/Engine/Source/Programs/BuildPatchTool/Private/BuildPatchTool.cpp
+++ b/Engine/Source/Programs/BuildPatchTool/Private/BuildPatchTool.cpp
@@ -1,76 +1,226 @@
 // Copyright Epic Games, Inc. All Rights Reserved.
 
 #include "BuildPatchTool.h"
-#include "UObject/Object.h"
+
 #include "RequiredProgramMainCPPInclude.h"
-#include "Interfaces/ToolMode.h"
+
+#include "Misc/App.h"
 #include "Misc/OutputDeviceError.h"
+#include "Http.h"
+#include "UObject/UObjectBase.h"
+
+#include "BuildPatchToolVersion.h"
+#include "Common/CmdLineProcessor.h"
+#include "Common/VerboseLogging.h"
+#include "Interfaces/IBuildPatchServicesModule.h"
+#include "Interfaces/ToolMode.h"
+#include "ToolModes/AutomationMode.h"
+#include "ToolModes/ChunkDeltaOptimiseMode.h"
+#include "ToolModes/ChunkDirectoryMode.h"
+#include "ToolModes/CompactifyMode.h"
+#include "ToolModes/DiffManifestMode.h"
+#include "ToolModes/ExtractMetadataMode.h"
+#include "ToolModes/EnumerationMode.h"
+#include "ToolModes/InstallManifestMode.h"
+#include "ToolModes/MergeManifestMode.h"
+#include "ToolModes/PackageChunksMode.h"
+#include "ToolModes/ToolModesHelp.h"
+#include "ToolModes/VerifyChunksMode.h"
+
+#if WITH_GOOGLE_MOCK
+THIRD_PARTY_INCLUDES_START
+#include "gmock/gmock.h"
+THIRD_PARTY_INCLUDES_END
+#endif
 
 using namespace BuildPatchTool;
 
 IMPLEMENT_APPLICATION(BuildPatchTool, "BuildPatchTool");
 DEFINE_LOG_CATEGORY(LogBuildPatchTool);
 
-const TCHAR* HandleLegacyCommandline(const TCHAR* CommandLine)
+namespace
 {
-	static FString CommandLineString;
-	CommandLineString = CommandLine;
-
-#if UE_BUILD_DEBUG
-	// Run smoke tests in debug
-	CommandLineString += TEXT(" -bForceSmokeTests ");
+	void SetCurrentWorkingDirectoryToBaseDir()
+	{
+		FPlatformMisc::CacheLaunchDir();
+#if PLATFORM_WINDOWS
+		TCHAR SystemError[1024];
+		verifyf(::SetCurrentDirectoryW(FWindowsPlatformProcess::BaseDir()), TEXT("Failed to set the working directory to '%s' (%s)"),
+			FWindowsPlatformProcess::BaseDir(),
+			FWindowsPlatformMisc::GetSystemErrorMessage(SystemError, UE_ARRAY_COUNT(SystemError), 0));
+#elif PLATFORM_MAC
+		chdir([FString(FMacPlatformProcess::BaseDir()).GetNSString() fileSystemRepresentation]);
+#elif PLATFORM_UNIX
+		chdir(TCHAR_TO_ANSI(FPlatformProcess::BaseDir()));
+#else
+#error Please implement SetCurrentWorkingDirectoryToBaseDir() for current platform
 #endif
+	}
+}
 
-	// For BPT, we will interpret -stdout as all log lines go to stdout.
-	if (FParse::Param(CommandLine, TEXT("stdout")))
+namespace CommandLineHelpers
+{
+	bool ParseSwitch(const TCHAR* InSwitch, FString& Value, const TArray<FString>& Switches)
 	{
-		CommandLineString += TEXT(" -AllowStdOutLogVerbosity ");
+		// Debug check requirements for InSwitch
+		checkSlow(InSwitch != nullptr);
+		checkSlow(InSwitch[FCString::Strlen(InSwitch) - 1] == TEXT('='));
+
+		for (const FString& Switch : Switches)
+		{
+			if (Switch.StartsWith(InSwitch))
+			{
+				FString StringValue;
+				Switch.Split(TEXT("="), nullptr, &StringValue);
+				Value = StringValue.TrimQuotes();
+				return true;
+			}
+		}
+		return false;
 	}
+}
+
+namespace WorkingDirectory
+{
+	// By default in UE4 projects the current working dir is the *.exe file location dir
+	// While using BPT it is more convenient to work with files from launch directory
+	// That's why BPT is built with DISABLE_CWD_CHANGES=1 define and the below code is added (EGS-23907)
+
+	static FString LaunchDirectoryPath;
 
-	// No longer supported options
-	if (CommandLineString.Contains(TEXT("-nochunks")))
+	void Configure()
 	{
-		UE_LOG(LogBuildPatchTool, Error, TEXT("NoChunks is no longer a supported mode. Remove this commandline option."));
-		return nullptr;
+		LaunchDirectoryPath = FPlatformProcess::GetCurrentWorkingDirectory() + TEXT("/");
+		SetCurrentWorkingDirectoryToBaseDir();
 	}
+}
 
-	// Check for legacy tool mode switching, if we don't have a mode and this was not a -help request, add the correct mode
-	if (!CommandLineString.Contains(TEXT("-mode=")) && !CommandLineString.Contains(TEXT("-help")))
+class FBuildPatchOutputDevice : public FOutputDevice
+{
+public:
+	FBuildPatchOutputDevice()
+		: bIsInitialized(false)
+	{}
+
+	void Initialize()
 	{
-		if (CommandLineString.Contains(TEXT("-compactify")))
-		{
-			CommandLineString = CommandLineString.Replace(TEXT("-compactify"), TEXT("-mode=compactify"));
-		}
-		else if (CommandLineString.Contains(TEXT("-dataenumerate")))
+		GConfig->GetArray(TEXT("BuildPatchTool.ConsoleLog"), TEXT("ExcludedCategories"), ExcludedCategories, GEngineIni);
+		bIsInitialized = true;
+	}
+
+	virtual void Serialize(const TCHAR* V, ELogVerbosity::Type Verbosity, const FName& Category) override
+	{
+		FString LogLine;
+		bool bShouldLog = false;
+		if (Verbosity <= ELogVerbosity::Display)
 		{
-			CommandLineString = CommandLineString.Replace(TEXT("-dataenumerate"), TEXT("-mode=enumeration"));
+			FString CategoryStr = Category.ToString();
+			// We should always log our own loggers display messages.
+			bShouldLog = CategoryStr.Equals(TEXT("LogBuildPatchTool"));
+			// Once we've been initialized, we should log things that we haven't explicitly excluded.
+			bShouldLog |= bIsInitialized && !ExcludedCategories.Contains(CategoryStr);
+			if (bShouldLog)
+			{
+				// We've disabled the console logger to give us finer grained control over our log output.
+				// This message is something we want to print, so we should display without the log "clutter".
+				LogLine = V;
+			}
 		}
-		// Patch generation did not have a mode flag, but does have some unique and required params
-		else if (CommandLineString.Contains(TEXT("-BuildRoot=")) && CommandLineString.Contains(TEXT("-BuildVersion=")))
+
+		if (bShouldLog)
 		{
-			FString NewCommandline(TEXT("-mode=patchgeneration "), CommandLineString.Len());
-			NewCommandline += CommandLineString;
-			CommandLineString = MoveTemp(NewCommandline);
+			FGenericPlatformMisc::LocalPrint(*(LogLine + TEXT("\n")));
+			fflush(stdout);
 		}
 	}
 
-	return *CommandLineString;
+private:
+	bool bIsInitialized;
+	TArray<FString> ExcludedCategories;
+};
+
+const FString GetRedactedCommandLine(const TCHAR* CommandLine)
+{
+    FString CommandLineString;
+    CommandLineString = CommandLine;
+    const FString Redaction = TEXT("******");
+    const TArray<FString> ParamsToRedact = { TEXT("ClientSecret") ,TEXT("EncryptionSecretKey") };
+    TArray<FString> Tokens, Switches;
+	FCommandLine::Parse(CommandLine, Tokens, Switches);
+    for (const FString& ParamToRedact : ParamsToRedact)
+    {
+        FString SearchString = ParamToRedact + TEXT("=");
+        FString Value;
+        if (CommandLineHelpers::ParseSwitch(*SearchString, Value, Switches))
+        {
+            CommandLineString.ReplaceInline(*(SearchString + Value), *(SearchString + Redaction));
+        }
+    }
+
+    return CommandLineString;
 }
 
-EReturnCode RunBuildPatchTool()
+struct FModuleLoadHelper
+{
+public:
+	FModuleLoadHelper(TCHAR const * const InModuleName)
+		: ModuleName(InModuleName)
+	{
+		FModuleManager::Get().LoadModuleChecked(ModuleName);
+	}
+
+	~FModuleLoadHelper()
+	{
+		FModuleManager::Get().UnloadModule(ModuleName);
+	}
+
+	template <typename ModuleInterface>
+	ModuleInterface& Get() const
+	{
+		return FModuleManager::GetModuleChecked<ModuleInterface>(ModuleName);
+	}
+
+	template <typename ModuleInterface>
+	TSharedRef<ModuleInterface>& GetShared() const
+	{
+		return FModuleManager::GetModuleChecked<ModuleInterface>(ModuleName);
+	}
+
+private:
+	TCHAR const * const ModuleName;
+};
+
+EReturnCode RunBuildPatchTool(const TCHAR* CommandLine)
 {
+	using namespace BuildPatchTool;
+#if !UE_BUILD_SHIPPING
+	if (FParse::Param(CommandLine, TEXT("debugging")))
+	{
+		while (!FPlatformMisc::IsDebuggerPresent())
+		{
+			FPlatformProcess::Sleep(0.01f); // to avoid 100% CPU load
+		}
+	}
+#endif
+
 	// Initialise the UObject module.
-	FModuleManager::Get().LoadModule(TEXT("CoreUObject"));
+	FModuleLoadHelper LoadedCoreUObjectModule(TEXT("CoreUObject"));
 	FCoreDelegates::OnInit.Broadcast();
 
 	// Load the BuildPatchServices Module.
-	IBuildPatchServicesModule& BuildPatchServicesModule = FModuleManager::LoadModuleChecked<IBuildPatchServicesModule>(TEXT("BuildPatchServices"));
+	FModuleLoadHelper LoadedBuildPatchServicesModule(TEXT("BuildPatchServices"));
+	IBuildPatchServicesModule& BuildPatchServicesModule = LoadedBuildPatchServicesModule.Get<IBuildPatchServicesModule>();
 
 	// Make sure we have processed UObjects from BPS.
 	ProcessNewlyLoadedUObjects();
 
-	// Instantiate and execute the tool.
-	TSharedRef<IToolMode> ToolMode = FToolModeFactory::Create(BuildPatchServicesModule);
+	// Unfortunately we cannot control UE logging going to GWarn as well as our own simpler log.
+	class FFeedbackContext* GWarnBackup = GWarn;
+	GWarn = nullptr;
+	ON_SCOPE_EXIT{ GWarn = GWarnBackup; };
+
+	// Run the tool mode.
+	TUniquePtr<IToolMode> ToolMode(FToolModeFactory::Create(BuildPatchServicesModule, CommandLine));
 	return ToolMode->Execute();
 }
 
@@ -78,8 +228,12 @@ int32 NumberOfWorkerThreadsDesired()
 {
 	const int32 MaxThreads = 64;
 	const int32 NumberOfCores = FPlatformMisc::NumberOfCores();
+	int32 NumThreadsOverride = NumberOfCores;
+	FParse::Value(FCommandLine::Get(), TEXT("NumWorkerThreads="), NumThreadsOverride);
+	int32 DesiredNumberOfThreads = FMath::Max(NumberOfCores - 1, NumThreadsOverride);
+	UE_LOG(LogBuildPatchTool, Log, TEXT("Desired number of threads is [%d] based on NumberOfCores [%d] and NumThreadsOverride [%d]."), DesiredNumberOfThreads, NumberOfCores, NumThreadsOverride);
 	// need to spawn at least one worker thread (see FTaskGraphImplementation)
-	return FMath::Max(FMath::Min(NumberOfCores - 1, MaxThreads), 1);
+	return FMath::Max(FMath::Min(DesiredNumberOfThreads, MaxThreads), 1);
 }
 
 void CheckAndReallocThreadPool()
@@ -88,6 +242,7 @@ void CheckAndReallocThreadPool()
 	{
 		const int32 ThreadsSpawned = GThreadPool->GetNumThreads();
 		const int32 DesiredThreadCount = NumberOfWorkerThreadsDesired();
+
 		if (ThreadsSpawned < DesiredThreadCount)
 		{
 			UE_LOG(LogBuildPatchTool, Log, TEXT("Engine only spawned %d worker threads, bumping up to %d!"), ThreadsSpawned, DesiredThreadCount);
@@ -102,27 +257,95 @@ void CheckAndReallocThreadPool()
 	}
 }
 
-EReturnCode BuildPatchToolMain(const TCHAR* CommandLine)
+EReturnCode BuildPatchToolMain(int32 ArgC, TCHAR* ArgV[])
 {
-	// Handle legacy commandlines
-	CommandLine = HandleLegacyCommandline(CommandLine);
-	if (CommandLine == nullptr)
+	FTaskTagScope Scope(ETaskTag::EGameThread);
+	FString EngineExitReason = TEXT("Exiting");
+	ON_SCOPE_EXIT
+	{
+		RequestEngineExit(*EngineExitReason);
+	};
+
+	GAlwaysReportCrash = true;
+	FApp::SetProjectName(TEXT("BuildPatchTool"));
+	FCmdLineProcessor CmdLineProcessor(ArgC, ArgV, WorkingDirectory::LaunchDirectoryPath, [](const FString& LogLine)
+	{
+		FGenericPlatformMisc::LocalPrint(*(LogLine + TEXT("\n")));
+		UE_LOG(LogBuildPatchTool, Error, TEXT("%s"), *LogLine);
+	});
+	FCommandLine::Set(*CmdLineProcessor.GetEngineString());
+
+	CmdLineProcessor.ProcessDragAndDrop();
+
+	if (!CmdLineProcessor.ProcessCommandlineFiles())
+	{
+		EngineExitReason = TEXT("BuildPatchToolMain ProcessCommandlineFileArgs failed.");
+		return EReturnCode::ArgumentProcessingError;
+	}
+
+	if (!CmdLineProcessor.HandleLegacyCommandline())
+	{
+		EngineExitReason = TEXT("BuildPatchToolMain HandleLegacyCommandline failed.");
+		return EReturnCode::ArgumentProcessingError;
+	}
+
+	if (!CmdLineProcessor.ConvertCommandlinePaths())
 	{
+		EngineExitReason = TEXT("BuildPatchToolMain ConvertCommandlinePaths failed.");
 		return EReturnCode::ArgumentProcessingError;
 	}
 
+	CmdLineProcessor.AddDesiredEngineFlags();
+
+	bool bVerboseLogsSpecified = CmdLineProcessor.ContainsFlag(TEXT("-verboseLogs"));
+	SetupVerboseLogging(bVerboseLogsSpecified);
+
+	// Add log device for stdout
+	bool bStdOutSpecified = CmdLineProcessor.ContainsFlag(TEXT("-stdout"));
+	FBuildPatchOutputDevice* BuildPatchOutputDevice = nullptr;
+	if (!bStdOutSpecified)
+	{
+		BuildPatchOutputDevice = new FBuildPatchOutputDevice();
+		GLog->AddOutputDevice(BuildPatchOutputDevice);
+	}
+
 	// Initialise application
-	GEngineLoop.PreInit(CommandLine);
-	UE_LOG(LogBuildPatchTool, Log, TEXT("Executed with commandline: %s"), CommandLine);
+	const FString FullCommandLine = CmdLineProcessor.GetEngineString();
+	GEngineLoop.PreInit(*FullCommandLine);
+	ON_SCOPE_EXIT
+	{
+		FEngineLoop::AppPreExit();
+		FModuleManager::Get().UnloadModulesAtShutdown();
+		FEngineLoop::AppExit();
+	};
+	if (BuildPatchOutputDevice != nullptr)
+	{
+		BuildPatchOutputDevice->Initialize();
+	}
+
+	// Explicitly set the current culture to the supported value
+	FInternationalization::Get().SetCurrentCulture("en");
+
+	UE_LOG(LogBuildPatchTool, Log, TEXT("BuildPatchTool v%s %s"), BUILDPATCHTOOL_VERSION_STRING, LexToString(FApp::GetBuildConfiguration()));
+	UE_LOG(LogBuildPatchTool, Log, TEXT("Received  commandline: %s"), *(GetRedactedCommandLine(*CmdLineProcessor.GetOriginalCommandLine())))
+	UE_LOG(LogBuildPatchTool, Log, TEXT("Processed commandline: %s"), *(GetRedactedCommandLine(*FullCommandLine)));
+
+	// Allow QA to test crashes are handled appropriately
+	bool bForceCrashSpecified = CmdLineProcessor.ContainsFlag(TEXT("-forceCrash"));
+	if (bForceCrashSpecified)
+	{
+		*((volatile int32*)0) = 0;
+	}
 
 	// Check whether as a program, we should bump up the number of threads in GThreadPool.
 	CheckAndReallocThreadPool();
 
 	// Run the application
-	EReturnCode ReturnCode = RunBuildPatchTool();
+	EReturnCode ReturnCode = RunBuildPatchTool(*FullCommandLine);
 	if (ReturnCode != EReturnCode::OK)
 	{
-		UE_LOG(LogBuildPatchTool, Error, TEXT("Tool exited with: %d"), (int32)ReturnCode);
+		UE_LOG(LogBuildPatchTool, Error, TEXT("Tool exited with %s (%d)"), LexToString(ReturnCode), (int32)ReturnCode);
+		UE_LOG(LogBuildPatchTool, Error, TEXT("For more details, see log file at %s"), *FPaths::ConvertRelativePathToFull(FPlatformOutputDevices::GetAbsoluteLogFilename()));
 	}
 
 	// Shutdown
@@ -135,34 +358,17 @@ EReturnCode BuildPatchToolMain(const TCHAR* CommandLine)
 	return ReturnCode;
 }
 
-const TCHAR* ProcessApplicationCommandline(int32 ArgC, TCHAR* ArgV[])
-{
-	static FString CommandLine = TEXT("-usehyperthreading -UNATTENDED");
-	for (int32 Option = 1; Option < ArgC; Option++)
-	{
-		CommandLine += TEXT(" ");
-		FString Argument(ArgV[Option]);
-		if (Argument.Contains(TEXT(" ")))
-		{
-			if (Argument.Contains(TEXT("=")))
-			{
-				FString ArgName;
-				FString ArgValue;
-				Argument.Split(TEXT("="), &ArgName, &ArgValue);
-				Argument = FString::Printf(TEXT("%s=\"%s\""), *ArgName, *ArgValue);
-			}
-			else
-			{
-				Argument = FString::Printf(TEXT("\"%s\""), *Argument);
-			}
-		}
-		CommandLine += Argument;
-	}
-	return *CommandLine;
-}
-
 INT32_MAIN_INT32_ARGC_TCHAR_ARGV()
 {
+
+#if WITH_GOOGLE_MOCK
+	int argc = 0;
+	wchar_t** argv = nullptr;
+	::testing::InitGoogleMock(&argc, argv);
+#endif
+	
+	WorkingDirectory::Configure();
+
 	EReturnCode ReturnCode;
 	// Using try&catch is the windows-specific method of interfacing with CrashReportClient
 #if PLATFORM_WINDOWS && !PLATFORM_SEH_EXCEPTIONS_DISABLED
@@ -172,7 +378,7 @@ INT32_MAIN_INT32_ARGC_TCHAR_ARGV()
 		// SetCrashHandler(nullptr) sets up default behavior for Linux and Mac interfacing with CrashReportClient
 		FPlatformMisc::SetCrashHandler(nullptr);
 		GIsGuarded = 1;
-		ReturnCode = BuildPatchToolMain(ProcessApplicationCommandline(ArgC, ArgV));
+		ReturnCode = BuildPatchToolMain(ArgC, ArgV);
 		GIsGuarded = 0;
 	}
 #if PLATFORM_WINDOWS && !PLATFORM_SEH_EXCEPTIONS_DISABLED
diff --git a/Engine/Source/Programs/BuildPatchTool/Private/ToolMode.cpp b/Engine/Source/Programs/BuildPatchTool/Private/ToolMode.cpp
--- a/Engine/Source/Programs/BuildPatchTool/Private/ToolMode.cpp
+++ b/Engine/Source/Programs/BuildPatchTool/Private/ToolMode.cpp
@@ -1,25 +1,21 @@
 // Copyright Epic Games, Inc. All Rights Reserved.
 
 #include "Interfaces/ToolMode.h"
-#include "Misc/CommandLine.h"
-#include "Interfaces/IBuildPatchServicesModule.h"
-#include "ToolModes/PatchGenerationMode.h"
-#include "ToolModes/ChunkDeltaOptimiseMode.h"
-#include "ToolModes/CompactifyMode.h"
-#include "ToolModes/EnumerationMode.h"
-#include "ToolModes/MergeManifestMode.h"
-#include "ToolModes/DiffManifestMode.h"
-#include "ToolModes/PackageChunksMode.h"
-#include "ToolModes/VerifyChunksMode.h"
-#include "ToolModes/AutomationMode.h"
+
 #include "BuildPatchTool.h"
+#include "BuildPatchToolVersion.h"
+#include "Common/MetadataSerialiser.h"
+#include "Interfaces/IBuildPatchServicesModule.h"
+#include "Misc/CommandLine.h"
+#include "ToolModes/ToolModesHelp.h"
 
 namespace BuildPatchTool
 {
 	class FHelpToolMode : public IToolMode
 	{
 	public:
-		FHelpToolMode()
+		FHelpToolMode(const TCHAR* InCommandLine)
+			: CommandLine(InCommandLine)
 		{}
 
 		virtual ~FHelpToolMode()
@@ -27,79 +23,52 @@ namespace BuildPatchTool
 
 		virtual EReturnCode Execute() override
 		{
-			bool bRequestedHelp = FParse::Param(FCommandLine::Get(), TEXT("help"));
+			bool bRequestedHelp = FParse::Param(CommandLine, TEXT("help"));
 
 			// Output generic help info
 			if (!bRequestedHelp)
 			{
-				UE_LOG(LogBuildPatchTool, Error, TEXT("No supported mode detected."));
+				UE_LOG(LogBuildPatchTool, Error, TEXT("ERROR: No mode was detected on the commandline, please check the provided args. You can use -help for information on available modes."));
+				UE_LOG(LogBuildPatchTool, Display, TEXT(""));
 			}
-			UE_LOG(LogBuildPatchTool, Display, TEXT("-help can be added with any mode selection to get extended information."));
-			UE_LOG(LogBuildPatchTool, Display, TEXT("Supported modes are:"));
-			UE_LOG(LogBuildPatchTool, Display, TEXT("  -mode=PatchGeneration      Mode that generates patch data for the a new build."));
-			UE_LOG(LogBuildPatchTool, Display, TEXT("  -mode=ChunkDeltaOptimise   Mode that produces a more optimal chunked patch between two specific builds."));
-			UE_LOG(LogBuildPatchTool, Display, TEXT("  -mode=Compactify           Mode that can clean up unneeded patch data from a given cloud directory with redundant data."));
-			UE_LOG(LogBuildPatchTool, Display, TEXT("  -mode=Enumeration          Mode that outputs the paths to referenced patch data given a single manifest."));
-			UE_LOG(LogBuildPatchTool, Display, TEXT("  -mode=MergeManifests       Mode that can combine two manifest files to create a new one, primarily used to create hotfixes."));
-			UE_LOG(LogBuildPatchTool, Display, TEXT("  -mode=DiffManifests        Mode that can diff two manifests and outputs what chunks would need to be downloaded and some stats."));
-			UE_LOG(LogBuildPatchTool, Display, TEXT("  -mode=PackageChunks        Mode that packages data required for an installation into larger files which can be used as local sources for build patch installers."));
-			UE_LOG(LogBuildPatchTool, Display, TEXT("  -mode=VerifyChunks         Mode that allows you to verify the integrity of patch data. It will load chunk or chunkdb files to verify they are not corrupt."));
-#if !UE_BUILD_SHIPPING
-			UE_LOG(LogBuildPatchTool, Display, TEXT("  -mode=AutomationTests      Mode that will run automation testing."));
-#endif // !UE_BUILD_SHIPPING
+			PrintHelp<FBptBaseToolModeHelp>();
 
 			// Error if this wasn't just a help request
 			return bRequestedHelp ? EReturnCode::OK : EReturnCode::UnknownToolMode;
 		}
+
+	private:
+		const TCHAR* CommandLine;
 	};
 
-	IToolModeRef FToolModeFactory::Create(IBuildPatchServicesModule& BpsInterface)
+	IToolMode* FToolModeFactory::Create(IBuildPatchServicesModule& BpsInterface, const TCHAR* CommandLine)
 	{
-		// Create the correct tool mode for the commandline given
+		TMap<FString, FToolModeCreateFunc> Modes;
+		GetModeRegistryDelegate().Broadcast(Modes);
+
+		// Create the correct tool mode for the commandline given.
 		FString ToolModeValue;
-		if (FParse::Value(FCommandLine::Get(), TEXT("mode="), ToolModeValue))
+		if (FParse::Value(CommandLine, TEXT("mode="), ToolModeValue))
 		{
-			if (ToolModeValue == TEXT("patchgeneration"))
-			{
-				return FPatchGenerationToolModeFactory::Create(BpsInterface);
-			}
-			else if (ToolModeValue == TEXT("chunkdeltaoptimise"))
-			{
-				return FChunkDeltaOptimiseToolModeFactory::Create(BpsInterface);
-			}
-			else if (ToolModeValue == TEXT("compactify"))
-			{
-				return FCompactifyToolModeFactory::Create(BpsInterface);
-			}
-			else if (ToolModeValue == TEXT("enumeration"))
-			{
-				return FEnumerationToolModeFactory::Create(BpsInterface);
-			}
-			else if (ToolModeValue == TEXT("mergemanifests"))
+			if (const FToolModeCreateFunc* ToolModeCreator = Modes.Find(*ToolModeValue))
 			{
-				return FMergeManifestToolModeFactory::Create(BpsInterface);
+				return (*ToolModeCreator)(BpsInterface, CommandLine);
 			}
-			else if (ToolModeValue == TEXT("diffmanifests"))
-			{
-				return FDiffManifestToolModeFactory::Create(BpsInterface);
-			}
-			else if (ToolModeValue == TEXT("packagechunks"))
-			{
-				return FPackageChunksToolModeFactory::Create(BpsInterface);
-			}
-			else if (ToolModeValue == TEXT("verifychunks"))
-			{
-				return FVerifyChunksToolModeFactory::Create(BpsInterface);
-			}
-#if !UE_BUILD_SHIPPING
-			else if (ToolModeValue == TEXT("automationtests"))
-			{
-				return FAutomationToolModeFactory::Create(BpsInterface);
-			}
-#endif // !UE_BUILD_SHIPPING
 		}
-		
-		// No mode provided so create the generic help, which will return ok if -help was provided else return UnknownToolMode error
-		return MakeShareable(new FHelpToolMode());
+
+		// If we didn't find a toolmode, check for empty mode name for overriden help.
+		if (const FToolModeCreateFunc* ToolModeCreator = Modes.Find(TEXT("")))
+		{
+			return (*ToolModeCreator)(BpsInterface, CommandLine);
+		}
+
+		// No previous attempts found a tool mode. Create the generic help, which will return ok if -help was provided else return UnknownToolMode error
+		return new FHelpToolMode(CommandLine);
+	}
+
+	FToolModeFactory::FModeRegistryDelegate& FToolModeFactory::GetModeRegistryDelegate()
+	{
+		static FToolModeFactory::FModeRegistryDelegate ModeRegistryDelegate;
+		return ModeRegistryDelegate;
 	}
 }
diff --git a/Engine/Source/Programs/BuildPatchTool/Private/ToolModes/AutomationMode.cpp b/Engine/Source/Programs/BuildPatchTool/Private/ToolModes/AutomationMode.cpp
--- a/Engine/Source/Programs/BuildPatchTool/Private/ToolModes/AutomationMode.cpp
+++ b/Engine/Source/Programs/BuildPatchTool/Private/ToolModes/AutomationMode.cpp
@@ -2,27 +2,74 @@
 
 #include "ToolModes/AutomationMode.h"
 
-#if !UE_BUILD_SHIPPING
+#if !UE_BUILD_SHIPPING && WITH_DEV_AUTOMATION_TESTS
 
-#include "Interfaces/IBuildPatchServicesModule.h"
 #include "BuildPatchTool.h"
-#include "Misc/CommandLine.h"
-#include "Misc/Paths.h"
-
+#include "Containers/Ticker.h"
 #include "IAutomationWorkerModule.h"
 #include "IAutomationControllerModule.h"
+#include "Interfaces/IBuildPatchServicesModule.h"
+#include "Misc/CommandLine.h"
+#include "Misc/Paths.h"
 #include "Modules/ModuleManager.h"
+#include "ToolModes/HelpToolMode.h"
+#include "ToolModes/ToolModesHelp.h"
 #include "UObject/UObjectGlobals.h"
-#include "Containers/Ticker.h"
+
+DECLARE_LOG_CATEGORY_EXTERN(LogAutomationMode, Warning, All);
+DEFINE_LOG_CATEGORY(LogAutomationMode);
 
 using namespace BuildPatchTool;
 
+#if WITH_GOOGLE_TEST
+
+THIRD_PARTY_INCLUDES_START
+#include "gtest/gtest.h"
+THIRD_PARTY_INCLUDES_END
+
+class FGoogleTestRunnerPrinter
+	: public ::testing::EmptyTestEventListener
+{
+	// Called after a failed assertion or a SUCCEED() invocation.
+	virtual void OnTestPartResult(const ::testing::TestPartResult& test_part_result) override
+	{
+		if (FAutomationTestBase* const CurrentTest = FAutomationTestFramework::GetInstance().GetCurrentTest())
+		{
+			if (test_part_result.passed())
+			{
+				CurrentTest->AddInfo(test_part_result.message());
+			}
+			else if (test_part_result.skipped())
+			{
+				CurrentTest->AddWarning(test_part_result.message());
+			}
+			else
+			{
+				CurrentTest->AddError(FString::Printf(TEXT("\n%s File:%s Line:%s"), test_part_result.message(), test_part_result.file_name(), test_part_result.line_number()));
+			}
+
+		}
+		else
+		{
+			UE_LOG(LogAutomationMode, Warning, TEXT("GoogleTest expectation or assert called when no automation test was actively running!"));
+		}
+	}
+};
+
+#endif
 class FAutomationToolMode : public IToolMode
 {
 public:
-	FAutomationToolMode(IBuildPatchServicesModule& InBpsInterface)
+	FAutomationToolMode(IBuildPatchServicesModule& InBpsInterface, const TCHAR* InCommandLine)
 		: BpsInterface(InBpsInterface)
+		, CommandLine(InCommandLine)
 	{
+#if WITH_GOOGLE_TEST
+		::testing::TestEventListeners& Listeners = ::testing::UnitTest::GetInstance()->listeners();
+		Listeners.Append(new FGoogleTestRunnerPrinter());
+		Listeners.Release(Listeners.default_result_printer());
+#endif
+
 	}
 
 	virtual ~FAutomationToolMode()
@@ -40,16 +87,7 @@ public:
 		// Print help if requested
 		if (bHelp)
 		{
-			UE_LOG(LogBuildPatchTool, Display, TEXT("AUTOMATION TEST MODE"));
-			UE_LOG(LogBuildPatchTool, Display, TEXT("This tool mode runs automation tests."));
-			UE_LOG(LogBuildPatchTool, Display, TEXT(""));
-			UE_LOG(LogBuildPatchTool, Display, TEXT("No arguments are required."));
-			UE_LOG(LogBuildPatchTool, Display, TEXT(""));
-			UE_LOG(LogBuildPatchTool, Display, TEXT("Optional arguments:"));
-			UE_LOG(LogBuildPatchTool, Display, TEXT("  -TestList=\"\"   Specifies in quotes, the list of tests to run. The list is + delimited."));
-			UE_LOG(LogBuildPatchTool, Display, TEXT(""));
-			UE_LOG(LogBuildPatchTool, Display, TEXT("NB: If -TestList is not specified, then all BuildPatchServices tests are ran."));
-			UE_LOG(LogBuildPatchTool, Display, TEXT(""));
+			PrintHelp<FAutomationToolModeHelp>();
 			return EReturnCode::OK;
 		}
 
@@ -129,7 +167,7 @@ private:
 	{
 #define PARSE_SWITCH(Switch) ParseSwitch(TEXT(#Switch "="), Switch, Switches)
 		TArray<FString> Tokens, Switches;
-		FCommandLine::Parse(FCommandLine::Get(), Tokens, Switches);
+		FCommandLine::Parse(CommandLine, Tokens, Switches);
 
 		bHelp = ParseOption(TEXT("help"), Switches);
 		if (bHelp)
@@ -138,7 +176,7 @@ private:
 		}
 		if (!PARSE_SWITCH(TestList) || TestList.Contains(TEXT(";")))
 		{
-			TestList = TEXT("BuildPatchServices");
+			TestList = TEXT("BuildPatchServices+BuildPatchTool");
 		}
 		TestList.InsertAt(0, TEXT("Automation RunTests "));
 
@@ -147,14 +185,13 @@ private:
 	}
 
 private:
+	static TCHAR const* const MODE_NAME;
 	IBuildPatchServicesModule& BpsInterface;
+	const TCHAR* CommandLine;
 	bool bHelp;
 	FString TestList;
 };
 
-BuildPatchTool::IToolModeRef BuildPatchTool::FAutomationToolModeFactory::Create(IBuildPatchServicesModule& BpsInterface)
-{
-	return MakeShareable(new FAutomationToolMode(BpsInterface));
-}
+IMPLEMENT_BPT_MODE(AutomationTests, FAutomationToolMode);
 
 #endif // !UE_BUILD_SHIPPING
diff --git a/Engine/Source/Programs/BuildPatchTool/Private/ToolModes/ChunkDeltaOptimiseMode.cpp b/Engine/Source/Programs/BuildPatchTool/Private/ToolModes/ChunkDeltaOptimiseMode.cpp
--- a/Engine/Source/Programs/BuildPatchTool/Private/ToolModes/ChunkDeltaOptimiseMode.cpp
+++ b/Engine/Source/Programs/BuildPatchTool/Private/ToolModes/ChunkDeltaOptimiseMode.cpp
@@ -1,20 +1,23 @@
 // Copyright Epic Games, Inc. All Rights Reserved.
 
 #include "ToolModes/ChunkDeltaOptimiseMode.h"
+
+#include "BuildPatchTool.h"
+#include "Common/DiffAbortThresholdArgument.h"
+#include "Interfaces/ToolMode.h"
 #include "Misc/CommandLine.h"
+#include "Misc/Optional.h"
 #include "Misc/Paths.h"
-#include "Interfaces/IBuildPatchServicesModule.h"
-#include "BuildPatchTool.h"
+#include "ToolModes/ToolModesHelp.h"
 
 using namespace BuildPatchTool;
 
 class FChunkDeltaOptimiseMode : public IToolMode
 {
 public:
-	FChunkDeltaOptimiseMode(IBuildPatchServicesModule& InBpsInterface)
+	FChunkDeltaOptimiseMode(IBuildPatchServicesModule& InBpsInterface, const TCHAR* InCommandLine)
 		: BpsInterface(InBpsInterface)
-		, ScanWindowSize(8191)
-		, OutputChunkSize(1024 * 1024)
+		, CommandLine(InCommandLine)
 	{}
 
 	virtual ~FChunkDeltaOptimiseMode()
@@ -22,7 +25,9 @@ public:
 
 	virtual EReturnCode Execute() override
 	{
-		// Parse commandline.
+		using namespace BuildPatchServices;
+
+		// Parse command line.
 		if (ProcessCommandline() == false)
 		{
 			return EReturnCode::ArgumentProcessingError;
@@ -31,44 +36,96 @@ public:
 		// Print help if requested.
 		if (bHelp)
 		{
-			UE_LOG(LogBuildPatchTool, Display, TEXT("CHUNK DELTA OPTIMISE MODE"));
-			UE_LOG(LogBuildPatchTool, Display, TEXT("This tool supports optimising chunk based patches to reduce the number of chunks required to download when patching between specific versions."));
-			UE_LOG(LogBuildPatchTool, Display, TEXT(""));
-			UE_LOG(LogBuildPatchTool, Display, TEXT("Required arguments:"));
-			UE_LOG(LogBuildPatchTool, Display, TEXT("  -mode=ChunkDeltaOptimise  Must be specified to launch the tool in chunk delta optimise mode."));
-			UE_LOG(LogBuildPatchTool, Display, TEXT("  -ManifestA=\"\"             Specifies in quotes the file path to the base manifest."));
-			UE_LOG(LogBuildPatchTool, Display, TEXT("  -ManifestB=\"\"             Specifies in quotes the file path to the update manifest."));
-			UE_LOG(LogBuildPatchTool, Display, TEXT(""));
-			UE_LOG(LogBuildPatchTool, Display, TEXT("Optional arguments:"));
-			UE_LOG(LogBuildPatchTool, Display, TEXT("  -CloudDir=\"\"                     Specifies in quotes the cloud directory where existing data will be recognized from, and new data added to. If not provided, location of ManifestB will be used."));
-			UE_LOG(LogBuildPatchTool, Display, TEXT("  -ScanWindowSize=1000000          Specifies in bytes, the scan window to use. Range accepted is 128KiB >= n >= 8Kb, defaults to 8191 (Closest prime to 8KiB)."));
-			UE_LOG(LogBuildPatchTool, Display, TEXT("  -OutputChunkSize=1000000         Specifies in bytes, the chunk size to save out unknown data with. Range accepted is 10MiB >= n >= 1MB, defaults to 1048576 (1MiB)."));
-			UE_LOG(LogBuildPatchTool, Display, TEXT("  -DiffAbortThreshold=10000000000  Specifies in bytes, an upper limit for original diffs to try to enhance."));
-			UE_LOG(LogBuildPatchTool, Display, TEXT("                                   This allows short circuiting lengthy optimisation attempts on large diffs which may not benefit. Range accepted is n >= 1GB, defaults to never aborting."));
-			UE_LOG(LogBuildPatchTool, Display, TEXT(""));
+			PrintHelp<FChunkDeltaOptimiseModeHelp>();
 			return EReturnCode::OK;
 		}
 
+		LogConfig();
+
+		// Register secrets.
+		for (const TPair<FGuid, TArray<uint8>>& IdKeyPair : EncryptionSecrets)
+		{
+			BpsInterface.RegisterEncryptionSecret(IdKeyPair.Key, IdKeyPair.Value);
+		}
+
+		// Set any commandline based config.
+		GConfig->SetInt(TEXT("Portal.BuildPatch"), TEXT("ChunkRetries"), MaxRetryCount, GEngineIni);
+
 		// Setup and run.
-		BuildPatchServices::FChunkDeltaOptimiserConfiguration Configuration;
+		FChunkDeltaOptimiserConfiguration Configuration;
 		Configuration.ManifestAUri = ManifestA;
 		Configuration.ManifestBUri = ManifestB;
-		Configuration.CloudDirectory = CloudDir;
+		Configuration.OutputCloudDirectory = CloudDir;
+		Configuration.DownloadCloudDirectories = DownloadCloudDirectories;
 		Configuration.ScanWindowSize = ScanWindowSize;
 		Configuration.OutputChunkSize = OutputChunkSize;
 		Configuration.DiffAbortThreshold = DiffAbortThreshold;
+		Configuration.DeltaFilenameTrailer = DeltaFilenameTrailer;
+		Configuration.SourceDetailsLogFilename = SourceDetailsLogFilename;
+		Configuration.ManifestBTagAliases = MoveTemp(ManifestBTagAliases);
+		Configuration.ManifestATagAliases = MoveTemp(ManifestATagAliases);
 
 		// Run the build generation.
-		bool bSuccess = BpsInterface.OptimiseChunkDelta(Configuration);
+		IChunkDeltaOptimiserRef ChunkDeltaOptimiser = BpsInterface.CreateChunkDeltaOptimiser(Configuration);
+
+		bool bSuccess = ChunkDeltaOptimiser->Run();
+
 		return bSuccess ? EReturnCode::OK : EReturnCode::ToolFailure;
 	}
 
 private:
+	// returns false on an error, not if the switch is present! Check OutMap for presence.
+	bool ParseManyToOneMap(const TCHAR* SwitchName, const TArray<FString>& Switches, TMap<FString, FString>& OutMap)
+	{
+		TArray<FString> RawStrings;
+		FString SwitchNameWithEq = FString(SwitchName) + TEXT("=");
+		if (ParseSwitches(*SwitchNameWithEq, RawStrings, Switches))
+		{
+			for (const FString& ManyToOneString : RawStrings)
+			{
+				FString From, To;
+				if (!ManyToOneString.Split(TEXT(";"), &From, &To))
+				{
+					UE_LOG(LogBuildPatchTool, Error, TEXT("%s must contain a ';' separating the 'from' tags and the 'to' tag (%s)"), SwitchName, *ManyToOneString);
+					return false;
+				}
+
+				TArray<FString> FromArray;
+				From.ParseIntoArray(FromArray, TEXT(","));
+				if (FromArray.Num() == 0)
+				{
+					UE_LOG(LogBuildPatchTool, Error, TEXT("%s must contain a comma separated list of 'from' tags (%s)"), SwitchName, *ManyToOneString);
+					return false;
+				}
+				if (To.Len() == 0)
+				{
+					UE_LOG(LogBuildPatchTool, Error, TEXT("%s must contain a 'to' tag after a ';' (%s)"), SwitchName, *ManyToOneString);
+					return false;
+				}
+
+				for (const FString& FromEntry : FromArray)
+				{
+					OutMap.Add(FromEntry, To);
+				}
+			}
+		}
+		return true;
+	}
+
 	bool ProcessCommandline()
 	{
 #define PARSE_SWITCH(Switch) ParseSwitch(TEXT(#Switch "="), Switch, Switches)
+#define _PARSE_REQUIRED_SWITCH(Switch, bExpect) \
+		if (PARSE_SWITCH(Switch) != bExpect) \
+		{ \
+			UE_LOG(LogBuildPatchTool, Error, #Switch SENTENCE_END); \
+			bParametersOk = false; \
+		}
+#define PARSE_REQUIRED_SWITCH(Switch) _PARSE_REQUIRED_SWITCH(Switch, true)
+#define PARSE_UNEXPECTED_SWITCH(Switch) _PARSE_REQUIRED_SWITCH(Switch, false)
+
 		TArray<FString> Tokens, Switches;
-		FCommandLine::Parse(FCommandLine::Get(), Tokens, Switches);
+		FCommandLine::Parse(CommandLine, Tokens, Switches);
 
 		bHelp = ParseOption(TEXT("help"), Switches);
 		if (bHelp)
@@ -76,35 +133,68 @@ private:
 			return true;
 		}
 
-		// Get all required parameters.
-		if (!(PARSE_SWITCH(ManifestA) && PARSE_SWITCH(ManifestB)))
+		// TODOBPTONLINE
+#undef SENTENCE_END
+#define SENTENCE_END TEXT(" is a required parameter for offline mode.")
+		bool bParametersOk = true;
+		PARSE_REQUIRED_SWITCH(ManifestA);
+		PARSE_REQUIRED_SWITCH(ManifestB);
+#undef SENTENCE_END
+#define SENTENCE_END TEXT(" should only be used when uploading to Epic services.")
+		FString BuildVersionA, BuildVersionB;
+		PARSE_UNEXPECTED_SWITCH(BuildVersionA);
+		PARSE_UNEXPECTED_SWITCH(BuildVersionB);
+#undef SENTENCE_END
+#define SENTENCE_END TEXT(" is only supported when Epic services are available. Download an official build from Epic to enable this functionality.")
+		FString ClientId, OrganizationId, ProductId;
+		PARSE_UNEXPECTED_SWITCH(ClientId);
+		PARSE_UNEXPECTED_SWITCH(OrganizationId);
+		PARSE_UNEXPECTED_SWITCH(ProductId);
+
+		if (!bParametersOk)
+		{
+			return false;
+		}
+
+		if (ManifestA.IsEmpty() || ManifestB.IsEmpty())
 		{
-			UE_LOG(LogBuildPatchTool, Error, TEXT("ManifestA and ManifestB are required parameters"));
+			UE_LOG(LogBuildPatchTool, Error, TEXT("ManifestA/B arguments can't be empty"));
 			return false;
 		}
+		
 		NormalizeUriFile(ManifestA);
 		NormalizeUriFile(ManifestB);
 
 		// Get optional values.
 		const bool bHasCloudDir = PARSE_SWITCH(CloudDir);
+		const bool bHasDownloadCloudDirectory = ParseSwitches(TEXT("DownloadCloudDirectory="), DownloadCloudDirectories, Switches);
+
+		PARSE_SWITCH(SourceDetailsLogFilename);
+		PARSE_SWITCH(DeltaFilenameTrailer);
 		PARSE_SWITCH(ScanWindowSize);
 		PARSE_SWITCH(OutputChunkSize);
+		PARSE_SWITCH(MaxRetryCount);
 		if (!bHasCloudDir)
 		{
 			CloudDir = FPaths::GetPath(ManifestB);
 		}
+		if (!bHasDownloadCloudDirectory)
+		{
+			DownloadCloudDirectories.Add(FPaths::GetPath(ManifestB));
+		}
+
 		NormalizeUriPath(CloudDir);
-		if (!PARSE_SWITCH(DiffAbortThreshold))
+		for (FString& DownloadCloudDirectory : DownloadCloudDirectories)
 		{
-			DiffAbortThreshold = TNumericLimits<int64>::Max();
+			NormalizeUriPath(DownloadCloudDirectory);
 		}
 
 		// Clamp ScanWindowSize to sane range.
 		const uint32 RequestedScanWindowSize = ScanWindowSize;
-		ScanWindowSize = FMath::Clamp<uint32>(ScanWindowSize, 8000, 128*1024);
+		ScanWindowSize = FMath::Clamp<uint32>(ScanWindowSize, 128, 128*1024);
 		if (RequestedScanWindowSize != ScanWindowSize)
 		{
-			UE_LOG(LogBuildPatchTool, Warning, TEXT("Requested -ScanWindowSize=%u is outside of allowed range 128KiB >= n >= 8Kb. Please update your arg to be within range. Continuing with %u."), RequestedScanWindowSize, ScanWindowSize);
+			UE_LOG(LogBuildPatchTool, Warning, TEXT("Requested -ScanWindowSize=%u is outside of allowed range 128KiB >= n >= 128b. Please update your arg to be within range. Continuing with %u."), RequestedScanWindowSize, ScanWindowSize);
 		}
 
 		// Clamp OutputChunkSize to sane range.
@@ -114,31 +204,109 @@ private:
 		{
 			UE_LOG(LogBuildPatchTool, Warning, TEXT("Requested -OutputChunkSize=%u is outside of allowed range 10MiB >= n >= 1MB. Please update your arg to be within range. Continuing with %u."), RequestedOutputChunkSize, OutputChunkSize);
 		}
+		DiffAbortThreshold = FDiffAbortThresholdArgument::Parse(*this, Switches);
 
-		// Clamp DiffAbortThreshold to sane min.
-		const uint64 RequestedDiffAbortThreshold = DiffAbortThreshold;
-		DiffAbortThreshold = FMath::Max<uint64>(DiffAbortThreshold, 1000000000);
-		if (RequestedDiffAbortThreshold != DiffAbortThreshold)
+		// Parse all the secret key and ID pairs.
+		FString ParseErrorMessage;
+		if (!ParsePairs(TEXT("EncryptionSecretId="), TEXT("EncryptionSecretKey="), Switches, EncryptionSecrets, ParseErrorMessage) && !ParseErrorMessage.IsEmpty())
 		{
-			UE_LOG(LogBuildPatchTool, Warning, TEXT("Requested -DiffAbortThreshold=%llu is below allowed minimum n >= 1GB. Please update your arg to be above limit. Continuing with %llu."), RequestedDiffAbortThreshold, DiffAbortThreshold);
+			UE_LOG(LogBuildPatchTool, Error, TEXT("%s"), *ParseErrorMessage);
+			return false;
+		}
+
+		// Parse any tag aliases we want to use. This is for constructing multiple optimized deltas
+		// where we know we have a larger installation set to draw from instead of just the same tag.
+		// Format is: -ManifestATagAliases="tag1,tag2,tag3,tag4;alias"
+		if (!ParseManyToOneMap(TEXT("ManifestATagAliases"), Switches, ManifestATagAliases))
+		{
+			return false; // already logged.
+		}
+		if (!ParseManyToOneMap(TEXT("ManifestBTagAliases"), Switches, ManifestBTagAliases))
+		{
+			return false; // already logged.
 		}
 
 		return true;
+#undef SENTENCE_END
+#undef PARSE_UNEXPECTED_SWITCH
+#undef PARSE_REQUIRED_SWITCH
+#undef _PARSE_REQUIRED_SWITCH
 #undef PARSE_SWITCH
 	}
 
+	void LogConfig() const
+	{
+		UE_LOG(LogBuildPatchTool, Log, TEXT("-----Configuration for ChunkDeltaOptimise------"));
+		UE_LOG(LogBuildPatchTool, Log, TEXT("   DiffAbortThreshold: %s"), *LexToString(DiffAbortThreshold));
+		UE_LOG(LogBuildPatchTool, Log, TEXT("   ManifestA: %s"), *ManifestA);
+		UE_LOG(LogBuildPatchTool, Log, TEXT("   ManifestB: %s"), *ManifestB);
+		UE_LOG(LogBuildPatchTool, Log, TEXT("   CloudDir: %s"), *CloudDir);
+		for (const FString& DownloadCloudDirectory : DownloadCloudDirectories)
+		{
+			UE_LOG(LogBuildPatchTool, Log, TEXT("   DownloadCloudDirectory: %s"), *DownloadCloudDirectory);
+		}
+		UE_LOG(LogBuildPatchTool, Log, TEXT("   ScanWindowSize: %u"), ScanWindowSize);
+		UE_LOG(LogBuildPatchTool, Log, TEXT("   OutputChunkSize: %u"), OutputChunkSize);
+		UE_LOG(LogBuildPatchTool, Log, TEXT("   MaxRetryCount: %u"), MaxRetryCount);
+		for (const TPair<FGuid, TArray<uint8>>& EncryptionSecret : EncryptionSecrets)
+		{
+			UE_LOG(LogBuildPatchTool, Log, TEXT("   AvailableSecretId: %s"), *EncryptionSecret.Key.ToString());
+		}
+		if (ManifestATagAliases.Num())
+		{
+			// For logging we reverse the map
+			TMap<FString, TArray<FString>> BackwardsMap;
+			for (const TPair<FString, FString>& Pair : ManifestATagAliases)
+			{
+				BackwardsMap.FindOrAdd(Pair.Value).Add(Pair.Key);
+			}
+			for (const TPair<FString, TArray<FString>>& TagAlias : BackwardsMap)
+			{
+				UE_LOG(LogBuildPatchTool, Log, TEXT("   ManifestATagAliases: %s -> %s"), *FString::Join(TagAlias.Value, TEXT(",")), *TagAlias.Key);
+			}
+		}
+		if (ManifestBTagAliases.Num())
+		{
+			// For logging we reverse the map
+			TMap<FString, TArray<FString>> BackwardsMap;
+			for (const TPair<FString, FString>& Pair : ManifestBTagAliases)
+			{
+				BackwardsMap.FindOrAdd(Pair.Value).Add(Pair.Key);
+			}
+			for (const TPair<FString, TArray<FString>>& TagAlias : BackwardsMap)
+			{
+				UE_LOG(LogBuildPatchTool, Log, TEXT("   ManifestBTagAliases: %s -> %s"), *FString::Join(TagAlias.Value, TEXT(",")), *TagAlias.Key);
+			}
+		}
+		if (DeltaFilenameTrailer.Len())
+		{
+			UE_LOG(LogBuildPatchTool, Log, TEXT("   DeltaFilenameTrailer: %s"), *DeltaFilenameTrailer);
+		}
+		if (SourceDetailsLogFilename.Len())
+		{
+			UE_LOG(LogBuildPatchTool, Log, TEXT("   SourceDetailsLogFilename: %s"), *SourceDetailsLogFilename);
+		}
+	}
+
 private:
+	static TCHAR const* const MODE_NAME;
 	IBuildPatchServicesModule& BpsInterface;
-	bool bHelp;
+	const TCHAR* CommandLine;
+	bool bHelp = false;
 	FString ManifestA;
 	FString ManifestB;
 	FString CloudDir;
-	uint32 ScanWindowSize;
-	uint32 OutputChunkSize;
-	uint64 DiffAbortThreshold;
+	TArray<FString> DownloadCloudDirectories;
+	uint32 ScanWindowSize = 8191U;
+	uint32 OutputChunkSize = 1024U * 1024U;
+	uint32 MaxRetryCount = 30;
+	TOptional<BuildPatchServices::FDiffAbortThreshold> DiffAbortThreshold;
+	TMap<FString, FString> ManifestBTagAliases;
+	TMap<FString, FString> ManifestATagAliases;
+	FString DeltaFilenameTrailer;
+	FString SourceDetailsLogFilename;
+	TMap<FGuid, TArray<uint8>> EncryptionSecrets;
 };
 
-BuildPatchTool::IToolModeRef BuildPatchTool::FChunkDeltaOptimiseToolModeFactory::Create(IBuildPatchServicesModule& BpsInterface)
-{
-	return MakeShareable(new FChunkDeltaOptimiseMode(BpsInterface));
-}
+IMPLEMENT_BPT_MODE(ChunkDeltaOptimise, FChunkDeltaOptimiseMode);
+IMPLEMENT_BPT_MODE_ALIAS(ChunkDeltaOptimize, FChunkDeltaOptimiseMode);
\ No newline at end of file
diff --git a/Engine/Source/Programs/BuildPatchTool/Private/ToolModes/ChunkDirectoryMode.cpp b/Engine/Source/Programs/BuildPatchTool/Private/ToolModes/ChunkDirectoryMode.cpp
new file mode 100644
--- /dev/null
+++ b/Engine/Source/Programs/BuildPatchTool/Private/ToolModes/ChunkDirectoryMode.cpp
@@ -0,0 +1,621 @@
+// Copyright Epic Games, Inc. All Rights Reserved.
+
+#include "ToolModes/ChunkDirectoryMode.h"
+
+#include "Algo/Transform.h"
+#include "BuildPatchFeatureLevel.h"
+#include "BuildPatchTool.h"
+#include "Common/MetadataSerialiser.h"
+#include "Interfaces/IBuildPatchServicesModule.h"
+#include "Interfaces/IDirectoryChunker.h"
+#include "Interfaces/ToolMode.h"
+#include "Misc/Base64.h"
+#include "Misc/CommandLine.h"
+#include "Misc/DefaultValueHelper.h"
+#include "Misc/FileHelper.h"
+#include "Misc/OutputDeviceRedirector.h"
+#include "Misc/Paths.h"
+#include "Serialization/MemoryWriter.h"
+#include "ToolModes/ToolModesHelp.h"
+
+class FChunkDirectoryToolMode : public BuildPatchTool::IToolMode
+{
+public:
+	FChunkDirectoryToolMode(IBuildPatchServicesModule& BpsInterface, const TCHAR* InCommandLine);
+	virtual BuildPatchTool::EReturnCode Execute() override;
+private:
+	bool ProcessCommandline();
+	void LogConfig() const;
+	void StoreMetadata(const IBuildManifest& Manifest) const;
+private:
+	FString GenerateRandomManifestFilename();
+	bool ParsePrereqIds(const FString& ParamValue, TSet<FString>& OutPrereqIds);
+	bool ParseCustomField(const FString& Switch, TMap<FString, FVariant>& Fields);
+
+private:
+	static TCHAR const* const MODE_NAME;
+	IBuildPatchServicesModule& BpsInterface;
+	const TCHAR* CommandLine;
+	bool bHelp;
+	FString FeatureLevel;
+	FString BuildRoot;
+	FString CloudDir;
+	FString AppId = TEXT("0");
+	FString ArtifactId;
+	FString BuildVersion;
+	FString AppLaunch;
+	FString AppArgs;
+	FString PrereqIds;
+	TSet<FString> PrereqIdsSet;
+	FString PrereqName;
+	FString PrereqPath;
+	FString PrereqArgs;
+	FString UninstallActionPath;
+	FString UninstallActionArgs;
+	FString FileList;
+	FString FileIgnoreList;
+	FString FileAttributeList;
+	FString DataAgeThreshold;
+	uint32 ChunkWindowSize = 1048576;
+	bool bIgnoreOtherWindowSizes = false;
+	bool bFollowSymlinks = false;
+	bool bResaveKnownChunks = false;
+	bool bSHA1Only = false;
+	TArray<uint8> EncryptionSecretKey;
+	FGuid EncryptionSecretId;
+	TMap<FString, FVariant> CustomFields;
+	FString OutputFilename;
+	FString MetadataOutput;
+	FString FileSortOrder;
+	FString ManifestSHA1Hash;
+	BuildPatchServices::IDirectoryChunkerPtr DirectoryChunker;
+};
+
+IMPLEMENT_BPT_MODE(ChunkBuildDirectory, FChunkDirectoryToolMode);
+
+FChunkDirectoryToolMode::FChunkDirectoryToolMode(IBuildPatchServicesModule& InBpsInterface, const TCHAR* InCommandLine)
+	: BpsInterface(InBpsInterface)
+	, CommandLine(InCommandLine)
+{}
+
+BuildPatchTool::EReturnCode FChunkDirectoryToolMode::Execute()
+{
+	// Parse commandline
+	if (ProcessCommandline() == false)
+	{
+		return BuildPatchTool::EReturnCode::ArgumentProcessingError;
+	}
+
+	// Print help if requested
+	if (bHelp)
+	{
+		PrintHelp<BuildPatchTool::FChunkDirectoryToolModeHelp>();
+		return BuildPatchTool::EReturnCode::OK;
+	}
+
+	LogConfig();
+
+	// Check build root exists
+	if (!FPaths::DirectoryExists(BuildRoot))
+	{
+		UE_LOG(LogBuildPatchTool, Error, TEXT("Specified build root directory does not exist."));
+		return BuildPatchTool::EReturnCode::DirectoryNotFound;
+	}
+
+	// Check existence of the file list
+	if (!FileList.IsEmpty() && !FPaths::FileExists(FileList))
+	{
+		UE_LOG(LogBuildPatchTool, Error, TEXT("Provided file list was not found %s"), *FileList);
+		return BuildPatchTool::EReturnCode::FileNotFound;
+	}
+
+	// Check existence of file ignore list
+	if (!FileIgnoreList.IsEmpty() && !FPaths::FileExists(FileIgnoreList))
+	{
+		UE_LOG(LogBuildPatchTool, Error, TEXT("Provided file ignore list was not found %s"), *FileIgnoreList);
+		return BuildPatchTool::EReturnCode::FileNotFound;
+	}
+
+	// Check existence of file attributes list
+	if (!FileAttributeList.IsEmpty() && !FPaths::FileExists(FileAttributeList))
+	{
+		UE_LOG(LogBuildPatchTool, Error, TEXT("Provided file attribute list was not found %s"), *FileAttributeList);
+		return BuildPatchTool::EReturnCode::FileNotFound;
+	}
+
+	// Default the OutputFilename if not provided
+	if (OutputFilename.IsEmpty())
+	{
+		OutputFilename = FDefaultValueHelper::RemoveWhitespaces(ArtifactId + BuildVersion) + TEXT(".manifest");
+	}
+	// Otherwise check the parameter
+	else
+	{
+		if (OutputFilename.Contains(TEXT("/")))
+		{
+			UE_LOG(LogBuildPatchTool, Error, TEXT("Provided OutputFilename should be clean filename only. Invalid arg: %s"), *OutputFilename);
+			return BuildPatchTool::EReturnCode::ArgumentProcessingError;
+		}
+	}
+
+	// Setup and run
+	BuildPatchServices::FDirectoryChunkerConfiguration Configuration;
+	LexFromString(Configuration.FeatureLevel, *FeatureLevel);
+	Configuration.RootDirectory = BuildRoot;
+	Configuration.AppId = TCString<TCHAR>::Atoi64(*AppId);
+	Configuration.AppName = ArtifactId;
+	Configuration.BuildVersion = BuildVersion;
+	Configuration.LaunchExe = AppLaunch;
+	Configuration.LaunchCommand = AppArgs;
+	Configuration.InputListFile = FileList;
+	Configuration.IgnoreListFile = FileIgnoreList;
+	Configuration.AttributeListFile = FileAttributeList;
+	Configuration.PrereqIds = PrereqIdsSet;
+	Configuration.PrereqName = PrereqName;
+	Configuration.PrereqPath = PrereqPath;
+	Configuration.PrereqArgs = PrereqArgs;
+	Configuration.UninstallActionPath = UninstallActionPath;
+	Configuration.UninstallActionArgs = UninstallActionArgs;
+	Configuration.DataAgeThreshold = TCString<TCHAR>::Atod(*DataAgeThreshold);
+	Configuration.bShouldHonorReuseThreshold = DataAgeThreshold.IsEmpty() == false;
+	Configuration.OutputChunkWindowSize = ChunkWindowSize;
+	Configuration.bShouldMatchAnyWindowSize = !bIgnoreOtherWindowSizes && Configuration.FeatureLevel >= BuildPatchServices::EFeatureLevel::VariableSizeChunks;
+	Configuration.bResaveKnownChunks = bResaveKnownChunks;
+	Configuration.CustomFields = CustomFields;
+	Configuration.CloudDirectory = CloudDir;
+	Configuration.OutputFilename = OutputFilename;
+	Configuration.bFollowSymlinks = bFollowSymlinks;
+	Configuration.bSHA1Only = bSHA1Only;
+	Configuration.EncryptionSecretKey = EncryptionSecretKey;
+	Configuration.EncryptionSecretId = EncryptionSecretId;
+
+	const uint32 DefaultChunkWindowSize = Configuration.OutputChunkWindowSize;
+
+	// Set the desired file sort order
+	if (!FileSortOrder.IsEmpty())
+	{
+		Configuration.FileSortOrder = BuildPatchServices::EFileSortOrder::InvalidOrMax;
+		LexFromString(Configuration.FileSortOrder, *FileSortOrder);
+		if (Configuration.FileSortOrder == BuildPatchServices::EFileSortOrder::InvalidOrMax)
+		{
+			UE_LOG(LogBuildPatchTool, Error, TEXT("Provided FileSortOrder is not recognised. Invalid arg: -FileSortOrder=%s"), *FileSortOrder);
+			return BuildPatchTool::EReturnCode::ArgumentProcessingError;
+		}
+	}
+
+	// Check feature compatibility
+	if (Configuration.FeatureLevel == BuildPatchServices::EFeatureLevel::Invalid)
+	{
+		UE_LOG(LogBuildPatchTool, Error, TEXT("Provided FeatureLevel is not recognised. Invalid arg: -FeatureLevel=%s"), *FeatureLevel);
+		return BuildPatchTool::EReturnCode::ArgumentProcessingError;
+	}
+	const bool bHasCustomFields = Configuration.CustomFields.Num() > 0;
+	if (Configuration.FeatureLevel < BuildPatchServices::EFeatureLevel::CustomFields && bHasCustomFields)
+	{
+		UE_LOG(LogBuildPatchTool, Error, TEXT("Invalid args: FeatureLevel %s is not compatible with Custom, CustomInt, or CustomFloat."), LexToString(Configuration.FeatureLevel));
+		return BuildPatchTool::EReturnCode::ArgumentProcessingError;
+	}
+	const bool bHasAnyPrereqInfo = !PrereqName.IsEmpty() || !PrereqPath.IsEmpty() || !PrereqArgs.IsEmpty();
+	if (Configuration.FeatureLevel < BuildPatchServices::EFeatureLevel::StoresPrerequisitesInfo && bHasAnyPrereqInfo)
+	{
+		UE_LOG(LogBuildPatchTool, Error, TEXT("Invalid args: FeatureLevel %s is not compatible with PrereqName, PrereqPath, or PrereqArgs."), LexToString(Configuration.FeatureLevel));
+		return BuildPatchTool::EReturnCode::ArgumentProcessingError;
+	}
+	const bool bHasPrereqIds = Configuration.PrereqIds.Num() > 0;
+	if (Configuration.FeatureLevel < BuildPatchServices::EFeatureLevel::StoresPrerequisiteIds && bHasPrereqIds)
+	{
+		UE_LOG(LogBuildPatchTool, Error, TEXT("Invalid args: FeatureLevel %s is not compatible with PrereqIds."), LexToString(Configuration.FeatureLevel));
+		return BuildPatchTool::EReturnCode::ArgumentProcessingError;
+	}
+	const bool bHasUninstallAction = !Configuration.UninstallActionPath.IsEmpty() || !Configuration.UninstallActionArgs.IsEmpty();
+	if (Configuration.FeatureLevel < BuildPatchServices::EFeatureLevel::StoresUninstallActions && bHasUninstallAction)
+	{
+		UE_LOG(LogBuildPatchTool, Error, TEXT("Invalid args: FeatureLevel %s is not compatible with either UninstallActionPath or UninstallActionArgs."), LexToString(Configuration.FeatureLevel));
+		return BuildPatchTool::EReturnCode::ArgumentProcessingError;
+	}
+	const bool bHasNonDefaultWindowSize = Configuration.OutputChunkWindowSize != DefaultChunkWindowSize;
+	if (Configuration.FeatureLevel < BuildPatchServices::EFeatureLevel::VariableSizeChunks && bHasNonDefaultWindowSize)
+	{
+		UE_LOG(LogBuildPatchTool, Error, TEXT("Invalid args: FeatureLevel %s is not compatible with -ChunkWindowSize=%u."), LexToString(Configuration.FeatureLevel), Configuration.OutputChunkWindowSize);
+		return BuildPatchTool::EReturnCode::ArgumentProcessingError;
+	}
+	const bool bUsingEncryption = Configuration.EncryptionSecretKey.Num() > 0 || Configuration.EncryptionSecretId.IsValid();
+	if (Configuration.FeatureLevel < BuildPatchServices::EFeatureLevel::ManifestEncryptionSupport && bUsingEncryption)
+	{
+		UE_LOG(LogBuildPatchTool, Error, TEXT("Invalid args: FeatureLevel %s is not compatible with encryption arguments."), LexToString(Configuration.FeatureLevel));
+		return BuildPatchTool::EReturnCode::ArgumentProcessingError;
+	}
+
+
+	bool bSuccess;
+
+	// Run the build generation
+	UE_LOG(LogBuildPatchTool, Display, TEXT("Beginning chunk generation of version %s of artifact %s. Build root: %s"), *Configuration.BuildVersion, *Configuration.AppName, *Configuration.RootDirectory);
+	DirectoryChunker = BpsInterface.CreateDirectoryChunker(Configuration);
+
+	DirectoryChunker->OnManifestGeneratedDecrypted().AddRaw(this, &FChunkDirectoryToolMode::StoreMetadata);
+
+	bSuccess = DirectoryChunker->Run();
+
+	if (!bSuccess)
+	{
+		UE_LOG(LogBuildPatchTool, Error, TEXT("An unexpected error occurred generating chunks."));
+		return BuildPatchTool::EReturnCode::ToolFailure;
+	}
+
+	UE_LOG(LogBuildPatchTool, Display, TEXT("Chunk generation complete."));
+
+	return BuildPatchTool::EReturnCode::OK;
+}
+
+bool FChunkDirectoryToolMode::ParsePrereqIds(const FString& PreReqIdString, TSet<FString>& OutPrereqIds)
+{
+	if (PreReqIdString.Contains(Constants::Slash)
+		|| PreReqIdString.Contains(Constants::Backslash)
+		|| PreReqIdString.Contains(Constants::DoubleQuote))
+	{
+		return false;
+	}
+
+	TArray<FString> ParamValues;
+	PreReqIdString.ParseIntoArray(ParamValues, *Constants::Comma);
+	for (FString& ParamValue : ParamValues)
+	{
+		ParamValue.TrimStartAndEndInline();
+	}
+	OutPrereqIds.Append(ParamValues);
+	return true;
+}
+
+bool FChunkDirectoryToolMode::ParseCustomField(const FString& Switch, TMap<FString, FVariant>& Fields)
+{
+	FString Type, Left, Right;
+
+	Switch.Split(Constants::Equals, &Type, &Right);
+	Type.ToLowerInline();
+	Right.Split(Constants::Equals, &Left, &Right);
+	Left.TrimStartAndEndInline();
+	if (Type.Equals(Constants::Custom, ESearchCase::CaseSensitive))
+	{
+		// No need to trim string fields
+		CustomFields.Add(Left, FVariant(Right));
+	}
+	else if (Type.Equals(Constants::CustomInt, ESearchCase::CaseSensitive))
+	{
+		Right.TrimStartAndEndInline();
+		if (!Right.IsNumeric())
+		{
+			UE_LOG(LogBuildPatchTool, Error, TEXT("An error occurred processing numeric token from commandline -%s"), *Switch);
+			return false;
+		}
+		CustomFields.Add(Left, FVariant(TCString<TCHAR>::Atoi64(*Right)));
+	}
+	else if (Type.Equals(Constants::CustomFloat, ESearchCase::CaseSensitive))
+	{
+		Right.TrimStartAndEndInline();
+		if (!Right.IsNumeric())
+		{
+			UE_LOG(LogBuildPatchTool, Error, TEXT("An error occurred processing numeric token from commandline -%s"), *Switch);
+			return false;
+		}
+		CustomFields.Add(Left, FVariant(TCString<TCHAR>::Atod(*Right)));
+	}
+
+	return true;
+}
+
+bool FChunkDirectoryToolMode::ProcessCommandline()
+{
+#define PARSE_SWITCH(Switch) ParseSwitch(TEXT(#Switch "="), Switch, Switches)
+#define PARSE_UNEXPECTED_SWITCH(Switch) \
+		if (PARSE_SWITCH(Switch) != false) \
+		{ \
+			UE_LOG(LogBuildPatchTool, Error, #Switch SENTENCE_END); \
+			bParametersOk = false; \
+		}
+
+	TArray<FString> Tokens, Switches;
+	FCommandLine::Parse(CommandLine, Tokens, Switches);
+
+	bHelp = ParseOption(TEXT("help"), Switches);
+	if (bHelp)
+	{
+		return true;
+	}
+
+	// Get all required parameters
+	if (!(PARSE_SWITCH(CloudDir)
+		&& PARSE_SWITCH(BuildRoot)
+		&& PARSE_SWITCH(ArtifactId)
+		&& PARSE_SWITCH(BuildVersion)
+		&& PARSE_SWITCH(AppLaunch)
+		&& PARSE_SWITCH(AppArgs)))
+	{
+		UE_LOG(LogBuildPatchTool, Error, TEXT("CloudDir, BuildRoot, ArtifactId, BuildVersion, AppLaunch, and AppArgs are required parameters"));
+		return false;
+	}
+
+	// Grab the FeatureLevel. This is required param - no defaults.
+	const bool bFeatureLevelPresent = PARSE_SWITCH(FeatureLevel);
+	FeatureLevel.TrimStartAndEndInline();
+
+	const BuildPatchServices::EFeatureLevel OnlineFeatureLevelNumeric = BuildPatchServices::EFeatureLevel::FirstOptimisedDelta;
+	BuildPatchServices::EFeatureLevel FeatureLevelToSetNumeric = BuildPatchServices::EFeatureLevel::Invalid;
+
+	if (!bFeatureLevelPresent)
+	{
+		UE_LOG(LogBuildPatchTool, Error, TEXT("FeatureLevel was not provided. Please provide the FeatureLevel commandline argument which matches the existing client support."));
+		return false;
+	}
+	else
+	{
+		LexFromString(FeatureLevelToSetNumeric, *FeatureLevel);
+		if (BuildPatchServices::EFeatureLevel::Invalid == FeatureLevelToSetNumeric)
+		{
+			UE_LOG(LogBuildPatchTool, Error, TEXT("Invalid arg: -FeatureLevel=%s. FeatureLevel was provided, but not recognized. Please provide the FeatureLevel commandline argument which matches the existing client support."), *FeatureLevel);
+			return false;
+		}
+	}
+
+	FeatureLevel = LexToString(FeatureLevelToSetNumeric);
+
+	NormalizeUriPath(CloudDir);
+	NormalizeUriPath(BuildRoot);
+	NormalizeUriFile(AppLaunch);
+
+	// Trim directory names, because low level APIs will report trailing whitespace as fine for existence checks, but fail on actually writing data.
+	CloudDir.TrimStartAndEndInline();
+	BuildRoot.TrimStartAndEndInline();
+
+	if (BuildRoot.IsEmpty())
+	{
+		UE_LOG(LogBuildPatchTool, Error, TEXT("No build root was provided."));
+		return false;
+	}
+
+	if (CloudDir.IsEmpty())
+	{
+		UE_LOG(LogBuildPatchTool, Error, TEXT("No cloud directory was provided."));
+		return false;
+	}
+
+	// Ensure CloudDir and BuildRoot paths are both absolute paths and normalized before comparing them to one another.
+	CloudDir = FPaths::ConvertRelativePathToFull(CloudDir);
+	BuildRoot = FPaths::ConvertRelativePathToFull(BuildRoot);
+
+	// Check that the CloudDir and BuildRoot directories are not identical.
+	if (CloudDir.Equals(BuildRoot))
+	{
+		UE_LOG(LogBuildPatchTool, Error, TEXT("The BuildRoot and CloudDir directories must be different."));
+		return false;
+	}
+
+	// Check that the CloudDir is not a child of the BuildRoot.
+	if (CloudDir.StartsWith(BuildRoot) && (CloudDir[BuildRoot.Len()] == TEXT('/')))
+	{
+		UE_LOG(LogBuildPatchTool, Error, TEXT("The CloudDir directory must not be a child of the BuildRoot directory."));
+		return false;
+	}
+
+	// Check that the BuildRoot is not a child of the CloudDir.
+	if (BuildRoot.StartsWith(CloudDir) && (BuildRoot[CloudDir.Len()] == TEXT('/')))
+	{
+		UE_LOG(LogBuildPatchTool, Error, TEXT("The BuildRoot directory must not be a child of the CloudDir directory."));
+		return false;
+	}
+
+	// Get optional parameters
+	bool bParametersOk = true;
+	FString ClientId, OrganizationId, ProductId;
+	// TODOBPTONLINE
+#define SENTENCE_END TEXT(" is only supported when Epic services are available. Download an official build from Epic to enable this functionality.")
+	PARSE_UNEXPECTED_SWITCH(ClientId);
+	PARSE_UNEXPECTED_SWITCH(OrganizationId);
+	PARSE_UNEXPECTED_SWITCH(ProductId);
+	if (!bParametersOk)
+	{
+		return false;
+	}
+#undef SENTENCE_END
+
+	PARSE_SWITCH(AppId);
+	PARSE_SWITCH(FileList);
+	PARSE_SWITCH(FileIgnoreList);
+	PARSE_SWITCH(FileAttributeList);
+	PARSE_SWITCH(PrereqIds);
+	PARSE_SWITCH(PrereqName);
+	PARSE_SWITCH(PrereqPath);
+	PARSE_SWITCH(PrereqArgs);
+	PARSE_SWITCH(UninstallActionPath);
+	PARSE_SWITCH(UninstallActionArgs);
+	PARSE_SWITCH(DataAgeThreshold);
+	PARSE_SWITCH(ChunkWindowSize);
+	bIgnoreOtherWindowSizes = ParseOption(TEXT("IgnoreOtherWindowSizes"), Switches);
+	PARSE_SWITCH(OutputFilename);
+	PARSE_SWITCH(MetadataOutput);
+	PARSE_SWITCH(FileSortOrder);
+	bFollowSymlinks = ParseOption(TEXT("FollowSymlinks"), Switches);
+	bResaveKnownChunks = ParseOption(TEXT("ResaveKnownChunks"), Switches);
+	bSHA1Only = ParseOption(TEXT("SHA1Only"), Switches);
+
+	// Parse encryption inputs
+	PARSE_SWITCH(EncryptionSecretKey);
+	FString EncryptionSecretIdString;
+	ParseSwitch(TEXT("EncryptionSecretId="), EncryptionSecretIdString, Switches);
+	ParseValue(EncryptionSecretIdString, EncryptionSecretId);
+
+	NormalizeUriFile(FileList);
+	NormalizeUriFile(FileIgnoreList);
+	NormalizeUriFile(FileAttributeList);
+	NormalizeUriFile(PrereqPath);
+	NormalizeUriFile(UninstallActionPath);
+	NormalizeUriFile(OutputFilename);
+	NormalizeUriFile(MetadataOutput);
+
+	// Check manifest file extension.
+	const TCHAR* ManifestExtension = TEXT(".manifest");
+	if (!OutputFilename.IsEmpty() && !OutputFilename.EndsWith(ManifestExtension))
+	{
+		OutputFilename += ManifestExtension;
+	}
+
+	// Clamp ChunkWindowSize to sane range.
+	const uint32 RequestedChunkWindowSize = ChunkWindowSize;
+	ChunkWindowSize = FMath::Clamp<uint32>(ChunkWindowSize, 32000, 10485760);
+	if (RequestedChunkWindowSize != ChunkWindowSize)
+	{
+		UE_LOG(LogBuildPatchTool, Warning, TEXT("Requested -ChunkWindowSize=%u is outside of allowed range 10485760 >= n >= 32000. Please update your args to be within range. Continuing with %u."), RequestedChunkWindowSize, ChunkWindowSize);
+	}
+
+	// Check numeric values
+	if (!AppId.IsEmpty() && (!AppId.IsNumeric() || TCString<TCHAR>::Atoi64(*AppId) < 0))
+	{
+		UE_LOG(LogBuildPatchTool, Error, TEXT("An error occurred processing numeric token from commandline -AppId=%s"), *AppId);
+		return false;
+	}
+	if (!DataAgeThreshold.IsEmpty() && (!DataAgeThreshold.IsNumeric() || TCString<TCHAR>::Atod(*DataAgeThreshold) < 0.0))
+	{
+		UE_LOG(LogBuildPatchTool, Error, TEXT("An error occurred processing numeric token from commandline -DataAgeThreshold=%s"), *DataAgeThreshold);
+		return false;
+	}
+
+	// Check the correct number of bytes was received for a secret key
+	switch (EncryptionSecretKey.Num())
+	{
+		case 0:
+		case 32:
+		break;
+
+		default:
+			UE_LOG(LogBuildPatchTool, Error, TEXT("For encryption, EncryptionSecretKey must be 256 bit."));
+			return false;
+	}
+
+	// Check that if an encryption ID was provided, that it was successfully parsed
+	if (EncryptionSecretIdString.Len() > 0 && !EncryptionSecretId.IsValid())
+	{
+		UE_LOG(LogBuildPatchTool, Error, TEXT("For encryption, EncryptionSecretId must be a valid 32 character hexidecimal GUID."));
+		return false;
+	}
+
+	// Check neither or both encryption values are received
+	if (EncryptionSecretKey.Num() > 0 != EncryptionSecretId.IsValid())
+	{
+		UE_LOG(LogBuildPatchTool, Error, TEXT("For encryption, both EncryptionSecretKey and EncryptionSecretId must be provided."));
+		return false;
+	}
+
+	// Get custom fields to add to manifest
+	// These are optional, but a failure to parse one is an error
+	for (const FString& Switch : Switches)
+	{
+		if (Switch.StartsWith(Constants::Custom) && !ParseCustomField(Switch, CustomFields))
+		{
+			return false;
+		}
+	}
+
+	if (!PrereqIds.IsEmpty() && !ParsePrereqIds(PrereqIds, PrereqIdsSet))
+	{
+		UE_LOG(LogBuildPatchTool, Error, TEXT("An error occurred processing comma-separated list from commandline -PrereqIds=%s"), *PrereqIds);
+		return false;
+	}
+
+	return true;
+#undef PARSE_SWITCH
+#undef PARSE_UNEXPECTED_SWITCH
+}
+
+void FChunkDirectoryToolMode::LogConfig() const
+{
+	UE_LOG(LogBuildPatchTool, Log, TEXT("-----Configuration for ChunkDirectory------"));
+	UE_LOG(LogBuildPatchTool, Log, TEXT("   ArtifactId: %s"), *ArtifactId);
+
+	UE_LOG(LogBuildPatchTool, Log, TEXT("   FeatureLevel: %s"), *FeatureLevel);
+	UE_LOG(LogBuildPatchTool, Log, TEXT("   BuildRoot: %s"), *BuildRoot);
+	UE_LOG(LogBuildPatchTool, Log, TEXT("   CloudDir: %s"), *CloudDir);
+	UE_LOG(LogBuildPatchTool, Log, TEXT("   AppId: %s"), *AppId);
+	UE_LOG(LogBuildPatchTool, Log, TEXT("   BuildVersion: %s"), *BuildVersion);
+	UE_LOG(LogBuildPatchTool, Log, TEXT("   AppLaunch: %s"), *AppLaunch);
+	UE_LOG(LogBuildPatchTool, Log, TEXT("   AppArgs: %s"), *AppArgs);
+
+	FString PrereqIdString = FString::Join(PrereqIdsSet.Array(), TEXT(","));
+	UE_LOG(LogBuildPatchTool, Log, TEXT("   PrereqIds: %s"), *PrereqIdString);
+	UE_LOG(LogBuildPatchTool, Log, TEXT("   PrereqName: %s"), *PrereqName);
+	UE_LOG(LogBuildPatchTool, Log, TEXT("   PrereqPath: %s"), *PrereqPath);
+	UE_LOG(LogBuildPatchTool, Log, TEXT("   PrereqArgs: %s"), *PrereqArgs);
+	UE_LOG(LogBuildPatchTool, Log, TEXT("   UninstallActionPath: %s"), *UninstallActionPath);
+	UE_LOG(LogBuildPatchTool, Log, TEXT("   UninstallActionArgs: %s"), *UninstallActionArgs);
+	UE_LOG(LogBuildPatchTool, Log, TEXT("   FileList: %s"), *FileList);
+	UE_LOG(LogBuildPatchTool, Log, TEXT("   FileIgnoreList: %s"), *FileIgnoreList);
+	UE_LOG(LogBuildPatchTool, Log, TEXT("   FileAttributeList: %s"), *FileAttributeList);
+	UE_LOG(LogBuildPatchTool, Log, TEXT("   DataAgeThreshold: %s"), *DataAgeThreshold);
+	UE_LOG(LogBuildPatchTool, Log, TEXT("   ChunkWindowSize: %u"), ChunkWindowSize);
+	UE_LOG(LogBuildPatchTool, Log, TEXT("   MetadataOutput: %s"), *MetadataOutput);
+	UE_LOG(LogBuildPatchTool, Log, TEXT("   bIgnoreOtherWindowSizes: %s"), bIgnoreOtherWindowSizes ? TEXT("true") : TEXT("false"));
+	UE_LOG(LogBuildPatchTool, Log, TEXT("   bFollowSymlinks: %s"), bFollowSymlinks ? TEXT("true") : TEXT("false"));
+	UE_LOG(LogBuildPatchTool, Log, TEXT("   bResaveKnownChunks: %s"), bResaveKnownChunks ? TEXT("true") : TEXT("false"));
+	UE_LOG(LogBuildPatchTool, Log, TEXT("   bSHA1Only: %s"), bSHA1Only ? TEXT("true") : TEXT("false"));
+	if (EncryptionSecretId.IsValid())
+	{
+		UE_LOG(LogBuildPatchTool, Log, TEXT("   EncryptionSecretId: %s"), *LexToString(EncryptionSecretId));
+	}
+
+	for (const TPair<FString, FVariant>& CustomField : CustomFields)
+	{
+		switch (CustomField.Value.GetType())
+		{
+		case EVariantTypes::String:
+			UE_LOG(LogBuildPatchTool, Log, TEXT("   Custom Field %s: %s"), *CustomField.Get<0>(), *CustomField.Get<1>().GetValue<FString>());
+			break;
+		case EVariantTypes::Int64:
+			UE_LOG(LogBuildPatchTool, Log, TEXT("   Custom Field %s: %lld"), *CustomField.Get<0>(), CustomField.Get<1>().GetValue<int64>());
+			break;
+		case EVariantTypes::Double:
+			UE_LOG(LogBuildPatchTool, Log, TEXT("   Custom Field %s: %f"), *CustomField.Get<0>(), CustomField.Get<1>().GetValue<double>());
+			break;
+
+		default:
+			break;
+		}
+	}
+}
+
+FString FChunkDirectoryToolMode::GenerateRandomManifestFilename()
+{
+	FString Result;
+	TArray<uint8> OutputFilenameBytes;
+	FMemoryWriter MemoryWriter(OutputFilenameBytes);
+
+	// First add some random bytes for security
+	for (uint8 i = 0; i < 6; ++i)
+	{
+		uint8 RandomByte = FMath::Rand();
+		MemoryWriter << RandomByte;
+	}
+
+	// Now add a GUID for uniqueness
+	FGuid OutputFilenameGuid = FGuid::NewGuid();
+	MemoryWriter << OutputFilenameGuid;
+
+	MemoryWriter.Close();
+
+	// Serialize to base 64, and replace unsafe filename chars with safer alternatives.
+	Result = FBase64::Encode(OutputFilenameBytes);
+	Result.ReplaceInline(TEXT("+"), TEXT("-"));
+	Result.ReplaceInline(TEXT("/"), TEXT("_"));
+	while (Result.RemoveFromEnd(TEXT("=")));
+	Result.Append(TEXT(".manifest"));
+	return Result;
+}
+
+void FChunkDirectoryToolMode::StoreMetadata(const IBuildManifest& Manifest) const
+{
+	if (!MetadataOutput.IsEmpty())
+	{
+		EMetadataOutputFormat OutputFormat = MetadataOutput.EndsWith(TEXT(".txt")) ? EMetadataOutputFormat::Human : EMetadataOutputFormat::Json;
+		FString OutputString = FMetadataSerialiser::SerialiseMetadata(Manifest, OutputFormat);
+		// Save the output
+		FFileHelper::SaveStringToFile(OutputString, *MetadataOutput);
+	}
+}
diff --git a/Engine/Source/Programs/BuildPatchTool/Private/ToolModes/CompactifyMode.cpp b/Engine/Source/Programs/BuildPatchTool/Private/ToolModes/CompactifyMode.cpp
--- a/Engine/Source/Programs/BuildPatchTool/Private/ToolModes/CompactifyMode.cpp
+++ b/Engine/Source/Programs/BuildPatchTool/Private/ToolModes/CompactifyMode.cpp
@@ -1,19 +1,22 @@
 // Copyright Epic Games, Inc. All Rights Reserved.
 
 #include "ToolModes/CompactifyMode.h"
-#include "Interfaces/IBuildPatchServicesModule.h"
+
+#include "BuildPatchSettings.h"
 #include "BuildPatchTool.h"
+#include "Interfaces/IBuildPatchServicesModule.h"
 #include "Misc/CommandLine.h"
 #include "Misc/Paths.h"
-#include "BuildPatchSettings.h"
+#include "ToolModes/ToolModesHelp.h"
 
 using namespace BuildPatchTool;
 
 class FCompactifyToolMode : public IToolMode
 {
 public:
-	FCompactifyToolMode(IBuildPatchServicesModule& InBpsInterface)
+	FCompactifyToolMode(IBuildPatchServicesModule& InBpsInterface, const TCHAR* InCommandLine)
 		: BpsInterface(InBpsInterface)
+		, CommandLine(InCommandLine)
 	{}
 
 	virtual ~FCompactifyToolMode()
@@ -30,23 +33,12 @@ public:
 		// Print help if requested
 		if (bHelp)
 		{
-			UE_LOG(LogBuildPatchTool, Display, TEXT("COMPACTIFY MODE"));
-			UE_LOG(LogBuildPatchTool, Display, TEXT("This tool supports the removal of redundant patch data from a cloud directory."));
-			UE_LOG(LogBuildPatchTool, Display, TEXT(""));
-			UE_LOG(LogBuildPatchTool, Display, TEXT("Required arguments:"));
-			UE_LOG(LogBuildPatchTool, Display, TEXT("  -mode=Compactify           Must be specified to launch the tool in compactify mode."));
-			UE_LOG(LogBuildPatchTool, Display, TEXT("  -CloudDir=\"\"               Specifies in quotes the cloud directory where manifest files and chunks to be compactified can be found."));
-			UE_LOG(LogBuildPatchTool, Display, TEXT("  -DataAgeThreshold=2        The minimum age in days of chunk files that will be deleted. Any unreferenced chunks older than this will be deleted."));
-			UE_LOG(LogBuildPatchTool, Display, TEXT(""));
-			UE_LOG(LogBuildPatchTool, Display, TEXT("Optional arguments:"));
-			UE_LOG(LogBuildPatchTool, Display, TEXT("  -preview                   Log all the actions it will take to update internal structures, but don't actually execute them."));
-			UE_LOG(LogBuildPatchTool, Display, TEXT("  -DeletedChunkLogFile=\"\"    Log the list of paths of deleted chunk files to this specified filename. All paths are relative to CloudDir."));
-			UE_LOG(LogBuildPatchTool, Display, TEXT(""));
-			UE_LOG(LogBuildPatchTool, Display, TEXT("NB: If -DataAgeThreshold is not supplied, then all unreferenced existing data is eligible for deletion by the compactify process."));
-			UE_LOG(LogBuildPatchTool, Display, TEXT(""));
+			PrintHelp<FCompactifyToolModeHelp>();
 			return EReturnCode::OK;
 		}
 
+		LogConfig();
+
 		// Grab options
 		BuildPatchServices::FCompactifyConfiguration Configuration;
 		Configuration.CloudDirectory = CloudDir;
@@ -65,7 +57,7 @@ private:
 	{
 #define PARSE_SWITCH(Switch) ParseSwitch(TEXT(#Switch "="), Switch, Switches)
 		TArray<FString> Tokens, Switches;
-		FCommandLine::Parse(FCommandLine::Get(), Tokens, Switches);
+		FCommandLine::Parse(CommandLine, Tokens, Switches);
 
 		bHelp = ParseOption(TEXT("help"), Switches);
 		if (bHelp)
@@ -80,10 +72,17 @@ private:
 			UE_LOG(LogBuildPatchTool, Error, TEXT("CloudDir and DataAgeThreshold are required parameters"));
 			return false;
 		}
+
+		if (CloudDir.IsEmpty())
+		{
+			UE_LOG(LogBuildPatchTool, Error, TEXT("CloudDir argument can't be empty"));
+			return false;
+		}
+
 		NormalizeUriPath(CloudDir);
 
 		// Check required numeric values
-		if (!DataAgeThreshold.IsNumeric())
+		if (!DataAgeThreshold.IsNumeric() || TCString<TCHAR>::Atod(*DataAgeThreshold) < 0.0)
 		{
 			UE_LOG(LogBuildPatchTool, Error, TEXT("An error occurred processing numeric token from commandline -DataAgeThreshold=%s"), *DataAgeThreshold);
 			return false;
@@ -98,8 +97,19 @@ private:
 #undef PARSE_SWITCH
 	}
 
+	void LogConfig() const
+	{
+		UE_LOG(LogBuildPatchTool, Log, TEXT("-----Configuration for CompactifyMode------"));
+		UE_LOG(LogBuildPatchTool, Log, TEXT("   CloudDir: %s"), *CloudDir);
+		UE_LOG(LogBuildPatchTool, Log, TEXT("   DataAgeThreshold: %s"), *DataAgeThreshold);
+		UE_LOG(LogBuildPatchTool, Log, TEXT("   DeletedChunkLogFile: %s"), *DeletedChunkLogFile);
+		UE_LOG(LogBuildPatchTool, Log, TEXT("   bPreview: %s"), bPreview ? TEXT("true") : TEXT("false"));
+	}
+
 private:
+	static TCHAR const* const MODE_NAME;
 	IBuildPatchServicesModule& BpsInterface;
+	const TCHAR* CommandLine;
 	bool bHelp;
 	FString CloudDir;
 	FString DataAgeThreshold;
@@ -107,7 +117,4 @@ private:
 	bool bPreview;
 };
 
-BuildPatchTool::IToolModeRef BuildPatchTool::FCompactifyToolModeFactory::Create(IBuildPatchServicesModule& BpsInterface)
-{
-	return MakeShareable(new FCompactifyToolMode(BpsInterface));
-}
+IMPLEMENT_BPT_MODE(Compactify, FCompactifyToolMode);
diff --git a/Engine/Source/Programs/BuildPatchTool/Private/ToolModes/DiffManifestMode.cpp b/Engine/Source/Programs/BuildPatchTool/Private/ToolModes/DiffManifestMode.cpp
--- a/Engine/Source/Programs/BuildPatchTool/Private/ToolModes/DiffManifestMode.cpp
+++ b/Engine/Source/Programs/BuildPatchTool/Private/ToolModes/DiffManifestMode.cpp
@@ -1,18 +1,22 @@
 // Copyright Epic Games, Inc. All Rights Reserved.
 
 #include "ToolModes/DiffManifestMode.h"
-#include "Interfaces/IBuildPatchServicesModule.h"
+
 #include "BuildPatchTool.h"
+#include "Interfaces/IBuildPatchServicesModule.h"
+#include "Interfaces/ToolMode.h"
 #include "Misc/CommandLine.h"
 #include "Misc/Paths.h"
+#include "ToolModes/ToolModesHelp.h"
 
 using namespace BuildPatchTool;
 
 class FDiffManifestToolMode : public IToolMode
 {
 public:
-	FDiffManifestToolMode(IBuildPatchServicesModule& InBpsInterface)
+	FDiffManifestToolMode(IBuildPatchServicesModule& InBpsInterface, const TCHAR* InCommandLine)
 		: BpsInterface(InBpsInterface)
+		, CommandLine(InCommandLine)
 	{}
 
 	virtual ~FDiffManifestToolMode()
@@ -29,60 +33,32 @@ public:
 		// Print help if requested.
 		if (bHelp)
 		{
-			UE_LOG(LogBuildPatchTool, Display, TEXT("DIFF MANIFEST MODE"));
-			UE_LOG(LogBuildPatchTool, Display, TEXT("This tool mode reports the changes between two existing manifest files."));
-			UE_LOG(LogBuildPatchTool, Display, TEXT(""));
-			UE_LOG(LogBuildPatchTool, Display, TEXT("Required arguments:"));
-			UE_LOG(LogBuildPatchTool, Display, TEXT("  -mode=DiffManifests    Must be specified to launch the tool in diff manifests mode."));
-			UE_LOG(LogBuildPatchTool, Display, TEXT("  -ManifestA=\"\"          Specifies in quotes the file path to the base manifest."));
-			UE_LOG(LogBuildPatchTool, Display, TEXT("  -ManifestB=\"\"          Specifies in quotes the file path to the update manifest."));
-			UE_LOG(LogBuildPatchTool, Display, TEXT(""));
-			UE_LOG(LogBuildPatchTool, Display, TEXT("Optional arguments:"));
-			UE_LOG(LogBuildPatchTool, Display, TEXT("  -InstallTagsA=\"\"       Specifies in quotes a comma seperated list of install tags used on ManifestA. You should include empty string if you want to count untagged files."));
-			UE_LOG(LogBuildPatchTool, Display, TEXT("                           Leaving the parameter out will use all files."));
-			UE_LOG(LogBuildPatchTool, Display, TEXT("                           -InstallTagsA=\"\" will be untagged files only."));
-			UE_LOG(LogBuildPatchTool, Display, TEXT("                           -InstallTagsA=\",tag\" will be untagged files plus files tagged with 'tag'."));
-			UE_LOG(LogBuildPatchTool, Display, TEXT("                           -InstallTagsA=\"tag\" will be files tagged with 'tag' only."));
-			UE_LOG(LogBuildPatchTool, Display, TEXT("  -InstallTagsB=\"\"       Specifies in quotes a comma seperated list of install tags used on ManifestB. Same rules apply as InstallTagsA."));
-			UE_LOG(LogBuildPatchTool, Display, TEXT("  -CompareTagSet=\"\"      Specifies in quotes a comma seperated list of install tags used to calculate differential statistics betweeen the manifests. Multiple lists are allowed. Same rules apply as InstallTagsA."));
-			UE_LOG(LogBuildPatchTool, Display, TEXT("  -OutputFile=\"\"         Specifies in quotes the file path where the diff will be exported as a JSON object."));
-			UE_LOG(LogBuildPatchTool, Display, TEXT("  -PatchDescriptorPath=\"\" If specified, will write out metadata files specifying what parts of each patched file changed."));
-			UE_LOG(LogBuildPatchTool, Display, TEXT("  -OnlyPatchDescriptors    If specified with PatchDescriptorPath, DiffManifests will only emit the files and not emit diff info."));
-			UE_LOG(LogBuildPatchTool, Display, TEXT("  -RequireOptimizedDelta   If specified, fail if an optimized delta file doesn't exist."));
-			UE_LOG(LogBuildPatchTool, Display, TEXT(""));
+			PrintHelp<FDiffManifestToolModeHelp>();
 			return EReturnCode::OK;
 		}
 
-		// Calc desired tags.
-		TSet<FString> TagSetA, TagSetB;
-		if (bHasTagsA)
-		{
-			TagSetA = ProcessTagList(InstallTagsA);
-		}
-		if (bHasTagsB)
-		{
-			TagSetB = ProcessTagList(InstallTagsB);
-		}
-
-		TArray<TSet<FString>> CompareTagSets;
+		LogConfig();
 
-		for (const FString& CompareTagsList : CompareTagsArray)
+		// Register secrets.
+		for (const TPair<FGuid, TArray<uint8>>& IdKeyPair : EncryptionSecrets)
 		{
-			CompareTagSets.Add(ProcessTagList(CompareTagsList));
+			BpsInterface.RegisterEncryptionSecret(IdKeyPair.Key, IdKeyPair.Value);
 		}
 
 		// Setup and run.
 		BuildPatchServices::FDiffManifestsConfiguration Configuration;
 		Configuration.ManifestAUri = ManifestA;
 		Configuration.ManifestBUri = ManifestB;
-		Configuration.TagSetA = TagSetA;
-		Configuration.TagSetB = TagSetB;
-		Configuration.CompareTagSets = CompareTagSets;
+		Configuration.TagSetA = InstallTagsA;
+		Configuration.TagSetB = InstallTagsB;
+		Configuration.CompareTagSetsA = CompareTagSets;
+		Configuration.CompareTagSetsB = CompareNewTagSets;
 		Configuration.OutputFilePath = OutputFile;
 		Configuration.OutputPatchDescriptorPath = PatchDescriptorPath;
 		Configuration.bOnlyPatchDescriptors = bOnlyPatchDescriptors;
 		Configuration.bRequireOptimizedDelta = bRequireOptimizedDelta;
-		Configuration.bEmitInstallTime = !bSkipInstallTime;
+		Configuration.DeltaFilenameTrailer = DeltaFilenameTrailer;
+		Configuration.bSimulateInstallTimes = bSimulateInstallTimes;
 
 		// Run the merge manifest routine.
 		bool bSuccess = BpsInterface.DiffManifests(Configuration);
@@ -94,8 +70,17 @@ private:
 	bool ProcessCommandLine()
 	{
 #define PARSE_SWITCH(Switch) ParseSwitch(TEXT(#Switch "="), Switch, Switches)
+#define _PARSE_REQUIRED_SWITCH(Switch, bExpect) \
+		if (PARSE_SWITCH(Switch) != bExpect) \
+		{ \
+			UE_LOG(LogBuildPatchTool, Error, #Switch SENTENCE_END); \
+			bParametersOk = false; \
+		}
+#define PARSE_REQUIRED_SWITCH(Switch) _PARSE_REQUIRED_SWITCH(Switch, true)
+#define PARSE_UNEXPECTED_SWITCH(Switch) _PARSE_REQUIRED_SWITCH(Switch, false)
+
 		TArray<FString> Tokens, Switches;
-		FCommandLine::Parse(FCommandLine::Get(), Tokens, Switches);
+		FCommandLine::Parse(CommandLine, Tokens, Switches);
 
 		bHelp = ParseOption(TEXT("help"), Switches);
 		if (bHelp)
@@ -103,32 +88,85 @@ private:
 			return true;
 		}
 
-		// Get all required parameters.
-		if (!(PARSE_SWITCH(ManifestA)
-		   && PARSE_SWITCH(ManifestB)))
+		// local variables to check for unexpected parameters
+		FString BuildVersionA, BuildVersionB, ClientId, OrganizationId, ProductId;
+#define SENTENCE_END TEXT(" is a required parameter for offline mode.")
+		bool bParametersOk = true;
+		PARSE_REQUIRED_SWITCH(ManifestA);
+		PARSE_REQUIRED_SWITCH(ManifestB);
+#undef SENTENCE_END
+#define SENTENCE_END TEXT(" should only be used when uploading to Epic services.")
+		PARSE_UNEXPECTED_SWITCH(BuildVersionA);
+		PARSE_UNEXPECTED_SWITCH(BuildVersionB);
+#undef SENTENCE_END
+#define SENTENCE_END TEXT(" is only supported when Epic services are available. Download an official build from Epic to enable this functionality.")
+		PARSE_UNEXPECTED_SWITCH(ClientId);
+		PARSE_UNEXPECTED_SWITCH(OrganizationId);
+		PARSE_UNEXPECTED_SWITCH(ProductId);
+#undef SENTENCE_END
+
+		if (!bParametersOk)
 		{
-			UE_LOG(LogBuildPatchTool, Error, TEXT("ManifestA and ManifestB are required parameters."));
 			return false;
 		}
+
+		if (ManifestA.IsEmpty() || ManifestB.IsEmpty())
+		{
+			UE_LOG(LogBuildPatchTool, Error, TEXT("ManifestA/B arguments can't be empty"));
+			return false;
+		}
+		
 		NormalizeUriFile(ManifestA);
 		NormalizeUriFile(ManifestB);
 
 		// Get optional parameters
-		bHasTagsA = PARSE_SWITCH(InstallTagsA);
-		bHasTagsB = PARSE_SWITCH(InstallTagsB);
+		PARSE_SWITCH(InstallTagsA);
+		PARSE_SWITCH(InstallTagsB);
 		PARSE_SWITCH(OutputFile);
 		NormalizeUriFile(OutputFile);
 
-		ParseSwitches(TEXT("CompareTagSet="), CompareTagsArray, Switches);
-
 		PARSE_SWITCH(PatchDescriptorPath);
 		NormalizeUriFile(PatchDescriptorPath);
 
+		PARSE_SWITCH(DeltaFilenameTrailer);
+		
 		bRequireOptimizedDelta = ParseOption(TEXT("RequireOptimizedDelta"), Switches);
 		bOnlyPatchDescriptors = ParseOption(TEXT("OnlyPatchDescriptors"), Switches);
-		bSkipInstallTime = ParseOption(TEXT("SkipInstallTime"), Switches);
+		bSimulateInstallTimes = ParseOption(TEXT("SimulateInstallTimes"), Switches);
+		bool bSkipInstallTime = ParseOption(TEXT("SkipInstallTime"), Switches);
+		if (bSkipInstallTime)
+		{
+			bSimulateInstallTimes = false;
+		}
+
+		TArray<FString> CompareTagsArray;
+		ParseSwitches(TEXT("CompareTagSet="), CompareTagsArray, Switches);
+
+		for (const FString& CompareTagsList : CompareTagsArray)
+		{
+			CompareTagSets.Add(ProcessTagList(CompareTagsList));
+		}
+
+		TArray<FString> CompareNewTagsArray;
+		ParseSwitches(TEXT("CompareNewTagSet="), CompareNewTagsArray, Switches);
+
+		for (const FString& CompareTagsList : CompareNewTagsArray)
+		{
+			CompareNewTagSets.Add(ProcessTagList(CompareTagsList));
+		}
+
+		// Parse all the secret key and ID pairs.
+		FString ParseErrorMessage;
+		if (!ParsePairs(TEXT("EncryptionSecretId="), TEXT("EncryptionSecretKey="), Switches, EncryptionSecrets, ParseErrorMessage) && !ParseErrorMessage.IsEmpty())
+		{
+			UE_LOG(LogBuildPatchTool, Error, TEXT("%s"), *ParseErrorMessage);
+			return false;
+		}
 
 		return true;
+#undef PARSE_UNEXPECTED_SWITCH
+#undef PARSE_REQUIRED_SWITCH
+#undef _PARSE_REQUIRED_SWITCH
 #undef PARSE_SWITCH
 	}
 
@@ -147,24 +185,56 @@ private:
 		return TSet<FString>(MoveTemp(TagArray));
 	}
 
+	void LogConfig() const
+	{
+		UE_LOG(LogBuildPatchTool, Log, TEXT("-----Configuration for DiffManifests------"));
+		UE_LOG(LogBuildPatchTool, Log, TEXT("   ManifestA: %s"), *ManifestA);
+		UE_LOG(LogBuildPatchTool, Log, TEXT("   ManifestB: %s"), *ManifestB);
+		TArray<FString> CompareTagSetStrings;
+		for (const TSet<FString>& CompareTagSet : CompareTagSets)
+		{
+			CompareTagSetStrings.Add(FString::Join(CompareTagSet, TEXT(",")));
+		}
+		TArray<FString> CompareNewTagSetStrings;
+		for (const TSet<FString>& CompareTagSet : CompareNewTagSets)
+		{
+			CompareNewTagSetStrings.Add(FString::Join(CompareTagSet, TEXT(",")));
+		}
+		UE_LOG(LogBuildPatchTool, Log, TEXT("   CompareTagSets: %s"), *FString::Join(CompareTagSetStrings, TEXT(";")));
+		UE_LOG(LogBuildPatchTool, Log, TEXT("   CompareNewTagSets: %s"), *FString::Join(CompareNewTagSetStrings, TEXT(";")));
+		UE_LOG(LogBuildPatchTool, Log, TEXT("   InstallTagsA: %s"), *FString::Join(InstallTagsA.Array(), TEXT(",")));
+		UE_LOG(LogBuildPatchTool, Log, TEXT("   InstallTagsB: %s"), *FString::Join(InstallTagsB.Array(), TEXT(",")));
+		UE_LOG(LogBuildPatchTool, Log, TEXT("   OutputFile: %s"), *OutputFile);
+		UE_LOG(LogBuildPatchTool, Log, TEXT("   RequireOptimizedDelta: %s"), *LexToString(bRequireOptimizedDelta));
+		UE_LOG(LogBuildPatchTool, Log, TEXT("   SimulateInstallTimes: %s"), *LexToString(bSimulateInstallTimes));
+		if (DeltaFilenameTrailer.Len())
+		{
+			UE_LOG(LogBuildPatchTool, Log, TEXT("   DeltaFilenameTrailer: %s"), *DeltaFilenameTrailer);		
+		}
+		for (const TPair<FGuid, TArray<uint8>>& EncryptionSecret : EncryptionSecrets)
+		{
+			UE_LOG(LogBuildPatchTool, Log, TEXT("   AvailableSecretId: %s"), *EncryptionSecret.Key.ToString());
+		}
+	}
+
 private:
+	static TCHAR const* const MODE_NAME;
 	IBuildPatchServicesModule& BpsInterface;
+	const TCHAR* CommandLine;
 	bool bHelp;
 	FString ManifestA;
 	FString ManifestB;
-	TArray<FString> CompareTagsArray;
-	bool bHasTagsA;
-	bool bHasTagsB;
-	FString InstallTagsA;
-	FString InstallTagsB;
+	TArray<TSet<FString>> CompareTagSets;
+	TArray<TSet<FString>> CompareNewTagSets;
+	TSet<FString> InstallTagsA;
+	TSet<FString> InstallTagsB;
 	FString OutputFile;
 	FString PatchDescriptorPath;
 	bool bRequireOptimizedDelta;
 	bool bOnlyPatchDescriptors;
-	bool bSkipInstallTime;
+	bool bSimulateInstallTimes;
+	FString DeltaFilenameTrailer;
+	TMap<FGuid, TArray<uint8>> EncryptionSecrets;
 };
 
-BuildPatchTool::IToolModeRef BuildPatchTool::FDiffManifestToolModeFactory::Create(IBuildPatchServicesModule& BpsInterface)
-{
-	return MakeShareable(new FDiffManifestToolMode(BpsInterface));
-}
+IMPLEMENT_BPT_MODE(DiffManifests, FDiffManifestToolMode);
diff --git a/Engine/Source/Programs/BuildPatchTool/Private/ToolModes/EnumerationMode.cpp b/Engine/Source/Programs/BuildPatchTool/Private/ToolModes/EnumerationMode.cpp
--- a/Engine/Source/Programs/BuildPatchTool/Private/ToolModes/EnumerationMode.cpp
+++ b/Engine/Source/Programs/BuildPatchTool/Private/ToolModes/EnumerationMode.cpp
@@ -1,18 +1,21 @@
 // Copyright Epic Games, Inc. All Rights Reserved.
 
 #include "ToolModes/EnumerationMode.h"
-#include "Interfaces/IBuildPatchServicesModule.h"
+
 #include "BuildPatchTool.h"
+#include "Interfaces/IBuildPatchServicesModule.h"
 #include "Misc/CommandLine.h"
 #include "Misc/Paths.h"
+#include "ToolModes/ToolModesHelp.h"
 
 using namespace BuildPatchTool;
 
 class FEnumerationToolMode : public IToolMode
 {
 public:
-	FEnumerationToolMode(IBuildPatchServicesModule& InBpsInterface)
+	FEnumerationToolMode(IBuildPatchServicesModule& InBpsInterface, const TCHAR* InCommandLine)
 		: BpsInterface(InBpsInterface)
+		, CommandLine(InCommandLine)
 	{}
 
 	virtual ~FEnumerationToolMode()
@@ -20,6 +23,8 @@ public:
 
 	virtual EReturnCode Execute() override
 	{
+		using namespace BuildPatchServices;
+
 		// Parse commandline
 		if (ProcessCommandline() == false)
 		{
@@ -29,35 +34,34 @@ public:
 		// Print help if requested
 		if (bHelp)
 		{
-			UE_LOG(LogBuildPatchTool, Display, TEXT("ENUMERATION MODE"));
-			UE_LOG(LogBuildPatchTool, Display, TEXT("This tool supports enumerating patch data referenced by a build manifest or chunkdb."));
-			UE_LOG(LogBuildPatchTool, Display, TEXT(""));
-			UE_LOG(LogBuildPatchTool, Display, TEXT("Required arguments:"));
-			UE_LOG(LogBuildPatchTool, Display, TEXT("  -mode=Enumeration    Must be specified to launch the tool in enumeration mode."));
-			UE_LOG(LogBuildPatchTool, Display, TEXT("  -InputFile=\"\"        Specifies in quotes the file path to the manifest or chunkdb to enumerate from."));
-			UE_LOG(LogBuildPatchTool, Display, TEXT("  -OutputFile=\"\"       Specifies in quotes the file path to a file where the list will be saved out."));
-			UE_LOG(LogBuildPatchTool, Display, TEXT(""));
-			UE_LOG(LogBuildPatchTool, Display, TEXT("Optional arguments:"));
-			UE_LOG(LogBuildPatchTool, Display, TEXT("  -includesizes        When specified, the size of each file in bytes will also be output (see notes)."));
-			UE_LOG(LogBuildPatchTool, Display, TEXT(""));
-			UE_LOG(LogBuildPatchTool, Display, TEXT("NB: If InputFile is a manifest, the output file format will be text file with one line per file, each containing cloud relative path."));
-			UE_LOG(LogBuildPatchTool, Display, TEXT("    e.g. path/to/file"));
-			UE_LOG(LogBuildPatchTool, Display, TEXT("    If InputFile is a chunk package, the output file format will be text file with one line per chunk, each containing tab separated hex chunk"));
-			UE_LOG(LogBuildPatchTool, Display, TEXT("    GUID, hex chunk rolling hash, and hex chunk SHA1."));
-			UE_LOG(LogBuildPatchTool, Display, TEXT("    e.g. 2CC26D05B64363780D5CF292E6B570A3\\t078070129133079067060057\\t527490FCA1DA6FAAB0E6F6E369E372FA693CCFBB"));
-			UE_LOG(LogBuildPatchTool, Display, TEXT("    Line endings are \\r\\n."));
-			UE_LOG(LogBuildPatchTool, Display, TEXT("NB: If includesizes is specified, each line of the output text file will end with a tab then the number of bytes of the file."));
+			PrintHelp<FEnumerationToolModeHelp>();
 			return EReturnCode::OK;
 		}
 
+		LogConfig();
+
+		// Register secrets.
+		for (const TPair<FGuid, TArray<uint8>>& IdKeyPair : EncryptionSecrets)
+		{
+			BpsInterface.RegisterEncryptionSecret(IdKeyPair.Key, IdKeyPair.Value);
+		}
+
+		// Check existence of input file
+		if (!FPaths::FileExists(InputFile))
+		{
+			UE_LOG(LogBuildPatchTool, Error, TEXT("Provided input file was not found: %s"), *InputFile);
+			return EReturnCode::FileNotFound;
+		}
+
 		// Setup and run.
-		BuildPatchServices::FPatchDataEnumerationConfiguration Settings;
+		FPatchDataEnumerationConfiguration Settings;
 		Settings.InputFile = InputFile;
 		Settings.OutputFile = OutputFile;
 		Settings.bIncludeSizes = bIncludeSizes;
 
 		// Run the enumeration routine
-		bool bSuccess = BpsInterface.EnumeratePatchData(Settings);
+		IPatchDataEnumerationRef PatchDataEnumeration = BpsInterface.CreatePatchDataEnumeration(Settings);
+		const bool bSuccess = PatchDataEnumeration->Run();
 		return bSuccess ? EReturnCode::OK : EReturnCode::ToolFailure;
 	}
 
@@ -67,7 +71,7 @@ private:
 	{
 #define PARSE_SWITCH(Switch) ParseSwitch(TEXT(#Switch "="), Switch, Switches)
 		TArray<FString> Tokens, Switches;
-		FCommandLine::Parse(FCommandLine::Get(), Tokens, Switches);
+		FCommandLine::Parse(CommandLine, Tokens, Switches);
 
 		bHelp = ParseOption(TEXT("help"), Switches);
 		if (bHelp)
@@ -75,39 +79,59 @@ private:
 			return true;
 		}
 
-		// Check for deprecated ManifestFile param.
-		FString ManifestFile;
-		if (PARSE_SWITCH(ManifestFile))
+		// Get all required parameters
+		if (!PARSE_SWITCH(InputFile) || InputFile.IsEmpty()
+		  || !PARSE_SWITCH(OutputFile) || OutputFile.IsEmpty())
 		{
-			InputFile = ManifestFile;
+			UE_LOG(LogBuildPatchTool, Error, TEXT("InputFile and OutputFile are required parameters"));
+			return false;
 		}
 
-		// Get all required parameters
-		if ((!PARSE_SWITCH(InputFile) && InputFile.IsEmpty())
-		  || !PARSE_SWITCH(OutputFile))
+		if (InputFile.Equals(OutputFile, ESearchCase::IgnoreCase))
 		{
-			UE_LOG(LogBuildPatchTool, Error, TEXT("InputFile and OutputFile are required parameters"));
+			UE_LOG(LogBuildPatchTool, Error, TEXT("InputFile and OutputFile cannot be the same"));
 			return false;
 		}
+
 		NormalizeUriFile(InputFile);
 		NormalizeUriFile(OutputFile);
 
 		// Get optional parameters
 		bIncludeSizes = ParseOption(TEXT("includesizes"), Switches);
 
+		// Parse all the secret key and ID pairs.
+		FString ParseErrorMessage;
+		if (!ParsePairs(TEXT("EncryptionSecretId="), TEXT("EncryptionSecretKey="), Switches, EncryptionSecrets, ParseErrorMessage) && !ParseErrorMessage.IsEmpty())
+		{
+			UE_LOG(LogBuildPatchTool, Error, TEXT("%s"), *ParseErrorMessage);
+			return false;
+		}
+
 		return true;
 #undef PARSE_SWITCH
 	}
 
+	void LogConfig() const
+	{
+		UE_LOG(LogBuildPatchTool, Log, TEXT("-----Configuration for EnumerationMode------"));
+		UE_LOG(LogBuildPatchTool, Log, TEXT("   InputFile: %s"), *InputFile);
+		UE_LOG(LogBuildPatchTool, Log, TEXT("   OutputFile: %s"), *OutputFile);
+		UE_LOG(LogBuildPatchTool, Log, TEXT("   bIncludeSizes: %s"), bIncludeSizes ? TEXT("true") : TEXT("false"));
+		for (const TPair<FGuid, TArray<uint8>>& EncryptionSecret : EncryptionSecrets)
+		{
+			UE_LOG(LogBuildPatchTool, Log, TEXT("   AvailableSecretId: %s"), *EncryptionSecret.Key.ToString());
+		}
+	}
+
 private:
+	static TCHAR const* const MODE_NAME;
 	IBuildPatchServicesModule& BpsInterface;
+	const TCHAR* CommandLine;
 	bool bHelp;
 	FString InputFile;
 	FString OutputFile;
 	bool bIncludeSizes;
+	TMap<FGuid, TArray<uint8>> EncryptionSecrets;
 };
 
-BuildPatchTool::IToolModeRef BuildPatchTool::FEnumerationToolModeFactory::Create(IBuildPatchServicesModule& BpsInterface)
-{
-	return MakeShareable(new FEnumerationToolMode(BpsInterface));
-}
+IMPLEMENT_BPT_MODE(Enumeration, FEnumerationToolMode);
