p4 integrate //UE5/Dev-Distro-5.5/Engine/Source/Programs/BuildPatchTool/... //UE5/Main/Engine/Source/Programs/BuildPatchTool/...

diff --git a/Engine/Source/Programs/BuildPatchTool/Private/ToolModes/InstallManifestMode.cpp b/Engine/Source/Programs/BuildPatchTool/Private/ToolModes/InstallManifestMode.cpp
new file mode 100644
--- /dev/null
+++ b/Engine/Source/Programs/BuildPatchTool/Private/ToolModes/InstallManifestMode.cpp
@@ -0,0 +1,360 @@
+// Copyright Epic Games, Inc. All Rights Reserved.
+
+#include "ToolModes/InstallManifestMode.h"
+
+#include "Algo/Transform.h"
+#include "Async/TaskGraphInterfaces.h"
+#include "BuildPatchTool.h"
+#include "Containers/Ticker.h"
+#include "HttpModule.h"
+#include "Interfaces/IHttpResponse.h"
+#include "Interfaces/ToolMode.h"
+#include "Interfaces/IBuildPatchServicesModule.h"
+#include "Misc/CommandLine.h"
+#include "Misc/OutputDeviceRedirector.h"
+#include "Misc/Paths.h"
+#include "ToolModes/ToolModesHelp.h"
+
+class FInstallManifestToolMode : public BuildPatchTool::IToolMode
+{
+public:
+	FInstallManifestToolMode(IBuildPatchServicesModule& BpsInterface, const TCHAR* CommandLine);
+	virtual BuildPatchTool::EReturnCode Execute() override;
+
+private:
+	void RequestManifestFile(const FString& ManifestUri);
+	void AppManifestHttpRequestComplete(FHttpRequestPtr HttpRequest, FHttpResponsePtr HttpResponse, bool bSucceeded, FString ManifestUri);
+	void FilterInstallation(const IBuildManifestRef& BuildManifest);
+	void StartInstallation(const IBuildManifestRef& BuildManifest);
+	void OnInstallComplete(const IBuildInstallerRef& Installer);
+	void ExitTool(BuildPatchTool::EReturnCode InReturnCode);
+	bool ProcessCommandline();
+	void LogConfig() const;
+
+private:
+	static TCHAR const* const MODE_NAME;
+	IBuildPatchServicesModule& BpsInterface;
+	const TCHAR* CommandLine;
+	BuildPatchTool::EReturnCode ReturnCode;
+	bool bHelp;
+	FString Manifest;
+	FString PrevManifest;
+	FString OutputDir;
+	TSet<FString> FileList;
+	TArray<FString> CloudDirs;
+	TMap<FGuid, TArray<uint8>> EncryptionSecrets;
+	IBuildManifestPtr InstallManifest;
+	IBuildManifestPtr CurrentManifest;
+	IBuildInstallerPtr BuildInstaller;
+	bool bRejectSymlinks = false;
+};
+
+FInstallManifestToolMode::FInstallManifestToolMode(IBuildPatchServicesModule& InBpsInterface, const TCHAR* InCommandLine)
+	: BpsInterface(InBpsInterface)
+	, CommandLine(InCommandLine)
+	, ReturnCode(BuildPatchTool::EReturnCode::OK)
+	, bHelp(false)
+{
+}
+
+BuildPatchTool::EReturnCode FInstallManifestToolMode::Execute()
+{
+	// Run any core initialisation required.
+	FHttpModule::Get();
+
+	// Parse commandline.
+	if (ProcessCommandline() == false)
+	{
+		return BuildPatchTool::EReturnCode::ArgumentProcessingError;
+	}
+
+	// Print help if requested.
+	if (bHelp)
+	{
+		PrintHelp<BuildPatchTool::FInstallManifestToolModeHelp>();
+		return BuildPatchTool::EReturnCode::OK;
+	}
+	LogConfig();
+
+	// Register secrets.
+	for (const TPair<FGuid, TArray<uint8>>& IdKeyPair : EncryptionSecrets)
+	{
+		BpsInterface.RegisterEncryptionSecret(IdKeyPair.Key, IdKeyPair.Value);
+	}
+
+	// Main timers.
+	double DeltaTime = 0.0;
+	double LastTime = FPlatformTime::Seconds();
+
+	// Setup desired frame times.
+	const float MainsFramerate = 100.0f;
+	const float MainsFrameTime = 1.0f / MainsFramerate;
+
+	// Kick off the first stage then run the tick loop.
+	if (!PrevManifest.IsEmpty())
+	{
+		RequestManifestFile(PrevManifest);
+	}
+	RequestManifestFile(Manifest);
+	while (!IsEngineExitRequested())
+	{
+		// Increment global frame counter once for each app tick.
+		GFrameCounter++;
+
+		// Application tick.
+		FTaskGraphInterface::Get().ProcessThreadUntilIdle(ENamedThreads::GameThread);
+		FTSTicker::GetCoreTicker().Tick(DeltaTime);
+
+		// Flush logs.
+		GLog->FlushThreadedLogs();
+
+		// Control frame rate.
+		FPlatformProcess::Sleep(FMath::Max<float>(0.0f, MainsFrameTime - (FPlatformTime::Seconds() - LastTime)));
+
+		// Calculate deltas.
+		const double AppTime = FPlatformTime::Seconds();
+		DeltaTime = AppTime - LastTime;
+		LastTime = AppTime;
+	}
+
+	return ReturnCode;
+}
+
+void FInstallManifestToolMode::RequestManifestFile(const FString& ManifestUri)
+{
+	const bool bIsHttpRequest = ManifestUri.StartsWith(TEXT("http"), ESearchCase::IgnoreCase);
+	if (bIsHttpRequest)
+	{
+		// Create a HTTP request.
+		TSharedRef<IHttpRequest, ESPMode::ThreadSafe> HttpRequest = FHttpModule::Get().CreateRequest();
+		HttpRequest->OnProcessRequestComplete().BindRaw(this, &FInstallManifestToolMode::AppManifestHttpRequestComplete, ManifestUri);
+		HttpRequest->SetURL(ManifestUri);
+		HttpRequest->SetVerb(TEXT("GET"));
+		if (!HttpRequest->ProcessRequest())
+		{
+			UE_LOG(LogBuildPatchTool, Error, TEXT("Unable to process manifest http request %s"), *ManifestUri);
+			ExitTool(BuildPatchTool::EReturnCode::ToolFailure);
+		}
+	}
+	else
+	{
+		// Read direct from file.
+		IBuildManifestPtr LoadedManifest = BpsInterface.LoadManifestFromFile(ManifestUri);
+		if (!LoadedManifest.IsValid())
+		{
+			UE_LOG(LogBuildPatchTool, Error, TEXT("Unable to read manifest from %s"), *ManifestUri);
+			ExitTool(BuildPatchTool::EReturnCode::InvalidData);
+		}
+		else if(ManifestUri == Manifest)
+		{
+			InstallManifest = LoadedManifest;
+		}
+		else
+		{
+			CurrentManifest = LoadedManifest;
+		}
+
+		// Are we done loading?
+		const bool bLoadComplete = (PrevManifest.IsEmpty() || CurrentManifest.IsValid()) && InstallManifest.IsValid();
+		if (bLoadComplete)
+		{
+			FilterInstallation(InstallManifest.ToSharedRef());
+		}
+	}
+}
+
+void FInstallManifestToolMode::AppManifestHttpRequestComplete(FHttpRequestPtr HttpRequest, FHttpResponsePtr HttpResponse, bool bSucceeded, FString ManifestUri)
+{
+	if (bSucceeded && HttpResponse.IsValid())
+	{
+		IBuildManifestPtr LoadedManifest = BpsInterface.MakeManifestFromData(HttpResponse->GetContent());
+		if (!LoadedManifest.IsValid())
+		{
+			UE_LOG(LogBuildPatchTool, Error, TEXT("Unable to serialise manifest from %s"), *HttpRequest->GetURL());
+			ExitTool(BuildPatchTool::EReturnCode::InvalidData);
+		}
+		else if(ManifestUri == Manifest)
+		{
+			InstallManifest = LoadedManifest;
+		}
+		else
+		{
+			CurrentManifest = LoadedManifest;
+		}
+
+		// Are we done loading?
+		const bool bLoadComplete = (PrevManifest.IsEmpty() || CurrentManifest.IsValid()) && InstallManifest.IsValid();
+		if (bLoadComplete)
+		{
+			FilterInstallation(InstallManifest.ToSharedRef());
+		}
+	}
+	else
+	{
+		UE_LOG(LogBuildPatchTool, Error, TEXT("Unable to download manifest from %s"), *HttpRequest->GetURL());
+		ExitTool(BuildPatchTool::EReturnCode::InvalidData);
+	}
+}
+
+void FInstallManifestToolMode::FilterInstallation(const IBuildManifestRef& BuildManifest)
+{
+	if (FileList.Num() > 0)
+	{
+		IBuildManifestPtr FilteredManifest = BpsInterface.MergeManifests(BuildManifest, BuildManifest, BuildManifest->GetVersionString(), TSet<FString>(), FileList);
+		if (FilteredManifest.IsValid())
+		{
+			StartInstallation(FilteredManifest.ToSharedRef());
+		}
+		else
+		{
+			UE_LOG(LogBuildPatchTool, Error, TEXT("Manifest filtering failed, check file list."));
+			ExitTool(BuildPatchTool::EReturnCode::ArgumentProcessingError);
+		}
+	}
+	else
+	{
+		StartInstallation(BuildManifest);
+	}
+}
+
+void FInstallManifestToolMode::StartInstallation(const IBuildManifestRef& BuildManifest)
+{
+	using namespace BuildPatchServices;
+
+	TSet<FGuid> AvailableSecrets;
+	Algo::Transform(EncryptionSecrets, AvailableSecrets, &TPair<FGuid, TArray<uint8>>::Key);
+	TSet<FGuid> NecessaryEncryptionSecrets = BuildManifest->GetNecessaryEncryptionSecretIds();
+
+	// Configure the installer.
+	FBuildInstallerConfiguration Configuration({ FInstallerAction::MakeInstallOrUpdate(CurrentManifest, BuildManifest) });
+	Configuration.InstallDirectory = OutputDir;
+	Configuration.StagingDirectory = FPaths::Combine(OutputDir, TEXT(".bptstage"));
+	Configuration.CloudDirectories = CloudDirs;
+	Configuration.EncryptionSecrets.Append(EncryptionSecrets);
+	Configuration.bRejectSymlinks = bRejectSymlinks;
+	if (NecessaryEncryptionSecrets.Difference(AvailableSecrets).Num() > 0)
+	{
+		Configuration.InstallMode = EInstallMode::Preload;
+	}
+
+	// Start the installer.
+	FBuildPatchInstallerDelegate InstallerCompleteDelegate = FBuildPatchInstallerDelegate::CreateRaw(this, &FInstallManifestToolMode::OnInstallComplete);
+	BuildInstaller = BpsInterface.CreateBuildInstaller(Configuration, MoveTemp(InstallerCompleteDelegate));
+	BuildInstaller->StartInstallation();
+
+}
+
+void FInstallManifestToolMode::OnInstallComplete(const IBuildInstallerRef& Installer)
+{
+	using namespace BuildPatchServices;
+	const FBuildInstallerConfiguration& InstallerConfiguration = Installer->GetConfiguration();
+	const FBuildInstallStats& InstallBuildStatistics = Installer->GetBuildStatistics();
+	const bool bSuccess = Installer->CompletedSuccessfully();
+
+	if (!bSuccess)
+	{
+		UE_LOG(LogBuildPatchTool, Display, TEXT("Install completed with failure. Error code: %s."), *Installer->GetErrorCode());
+		ExitTool(BuildPatchTool::EReturnCode::ToolFailure);
+	}
+	else
+	{
+		UE_LOG(LogBuildPatchTool, Display, TEXT("Install completed successfully."));
+		ExitTool(BuildPatchTool::EReturnCode::OK);
+	}
+}
+
+void FInstallManifestToolMode::ExitTool(BuildPatchTool::EReturnCode InReturnCode)
+{
+	ReturnCode = InReturnCode;
+	RequestEngineExit(TEXT("Installer complete."));
+}
+
+bool FInstallManifestToolMode::ProcessCommandline()
+{
+#define PARSE_SWITCH(Switch) ParseSwitch(TEXT(#Switch "="), Switch, Switches)
+
+	TArray<FString> Tokens, Switches;
+	FCommandLine::Parse(CommandLine, Tokens, Switches);
+
+	bHelp = ParseOption(TEXT("help"), Switches);
+	if (bHelp)
+	{
+		return true;
+	}
+
+	PARSE_SWITCH(Manifest);
+	PARSE_SWITCH(PrevManifest);
+	PARSE_SWITCH(OutputDir);
+	PARSE_SWITCH(FileList);
+	PARSE_SWITCH(CloudDirs);
+
+	bRejectSymlinks = ParseOption(TEXT("RejectSymlinks"), Switches);
+
+	NormalizeUriFile(Manifest);
+	NormalizeUriFile(PrevManifest);
+	NormalizeUriPath(OutputDir);
+	for (FString& File : FileList)
+	{
+		NormalizeUriPath(File);
+	}
+	for (FString& CloudDir : CloudDirs)
+	{
+		NormalizeUriFile(CloudDir);
+	}
+	if (CloudDirs.Num() == 0)
+	{
+		CloudDirs.Add(FPaths::GetPath(Manifest));
+	}
+
+	// Parse all the secret key and ID pairs.
+	FString ParseErrorMessage;
+	if (!ParsePairs(TEXT("EncryptionSecretId="), TEXT("EncryptionSecretKey="), Switches, EncryptionSecrets, ParseErrorMessage) && !ParseErrorMessage.IsEmpty())
+	{
+		UE_LOG(LogBuildPatchTool, Error, TEXT("%s"), *ParseErrorMessage);
+		return false;
+	}
+
+#undef PARSE_SWITCH
+	return true;
+}
+
+void FInstallManifestToolMode::LogConfig() const
+{
+	UE_LOG(LogBuildPatchTool, Log, TEXT("-----Configuration for InstallManifest------"));
+	UE_LOG(LogBuildPatchTool, Log, TEXT("  PrevManifest:   %s"), *PrevManifest);
+	UE_LOG(LogBuildPatchTool, Log, TEXT("  Manifest:       %s"), *Manifest);
+	UE_LOG(LogBuildPatchTool, Log, TEXT("  RejectSymlinks: %s"), bRejectSymlinks ? TEXT("true") : TEXT("false"));
+	UE_LOG(LogBuildPatchTool, Log, TEXT("  OutputDir:      %s"), *OutputDir);
+	bool bFirst = true;
+	for (const FString& File : FileList)
+	{
+		if (bFirst)
+		{
+			bFirst = false;
+			UE_LOG(LogBuildPatchTool, Log, TEXT("  FileList:     %s"), *File);
+		}
+		else
+		{
+			UE_LOG(LogBuildPatchTool, Log, TEXT("                %s"), *File);
+		}
+	}
+	bFirst = true;
+	for (const FString& CloudDir : CloudDirs)
+	{
+		if (bFirst)
+		{
+			bFirst = false;
+			UE_LOG(LogBuildPatchTool, Log, TEXT("  CloudDirs:      %s"), *CloudDir);
+		}
+		else
+		{
+			UE_LOG(LogBuildPatchTool, Log, TEXT("                  %s"), *CloudDir);
+		}
+	}
+	for (const TPair<FGuid, TArray<uint8>>& EncryptionSecret : EncryptionSecrets)
+	{
+		UE_LOG(LogBuildPatchTool, Log, TEXT("  AvailableSecretId: %s"), *EncryptionSecret.Key.ToString());
+	}
+}
+
+IMPLEMENT_BPT_MODE(InstallManifest, FInstallManifestToolMode);
diff --git a/Engine/Source/Programs/BuildPatchTool/Private/ToolModes/MergeManifestMode.cpp b/Engine/Source/Programs/BuildPatchTool/Private/ToolModes/MergeManifestMode.cpp
--- a/Engine/Source/Programs/BuildPatchTool/Private/ToolModes/MergeManifestMode.cpp
+++ b/Engine/Source/Programs/BuildPatchTool/Private/ToolModes/MergeManifestMode.cpp
@@ -1,18 +1,21 @@
 // Copyright Epic Games, Inc. All Rights Reserved.
 
 #include "ToolModes/MergeManifestMode.h"
-#include "Interfaces/IBuildPatchServicesModule.h"
+
 #include "BuildPatchTool.h"
+#include "Interfaces/IBuildPatchServicesModule.h"
 #include "Misc/CommandLine.h"
 #include "Misc/Paths.h"
+#include "ToolModes/ToolModesHelp.h"
 
 using namespace BuildPatchTool;
 
 class FMergeManifestToolMode : public IToolMode
 {
 public:
-	FMergeManifestToolMode(IBuildPatchServicesModule& InBpsInterface)
+	FMergeManifestToolMode(IBuildPatchServicesModule& InBpsInterface, const TCHAR* InCommandLine)
 		: BpsInterface(InBpsInterface)
+		, CommandLine(InCommandLine)
 	{}
 
 	virtual ~FMergeManifestToolMode()
@@ -29,25 +32,18 @@ public:
 		// Print help if requested
 		if (bHelp)
 		{
-			UE_LOG(LogBuildPatchTool, Display, TEXT("MERGE MANIFEST MODE"));
-			UE_LOG(LogBuildPatchTool, Display, TEXT("This tool supports generating a hotfix manifest from two existing manifest files."));
-			UE_LOG(LogBuildPatchTool, Display, TEXT(""));
-			UE_LOG(LogBuildPatchTool, Display, TEXT("Required arguments:"));
-			UE_LOG(LogBuildPatchTool, Display, TEXT("  -mode=MergeManifests    Must be specified to launch the tool in merge manifests mode."));
-			UE_LOG(LogBuildPatchTool, Display, TEXT("  -ManifestA=\"\"           Specifies in quotes the file path to the base manifest."));
-			UE_LOG(LogBuildPatchTool, Display, TEXT("  -ManifestB=\"\"           Specifies in quotes the file path to the update manifest."));
-			UE_LOG(LogBuildPatchTool, Display, TEXT("  -ManifestC=\"\"           Specifies in quotes the file path to the output manifest."));
-			UE_LOG(LogBuildPatchTool, Display, TEXT("  -BuildVersion=\"\"        Specifies in quotes the new version string for the build being produced."));
-			UE_LOG(LogBuildPatchTool, Display, TEXT(""));
-			UE_LOG(LogBuildPatchTool, Display, TEXT("Optional arguments:"));
-			UE_LOG(LogBuildPatchTool, Display, TEXT("  -MergeFileList=\"\"       Specifies in quotes, the path to a text file containing complete list of desired build root relative files followed by \\t character, followed by A or B to select the manifest to pull from. These should be seperated by \\r\\n line endings."));
-			UE_LOG(LogBuildPatchTool, Display, TEXT(""));
-			UE_LOG(LogBuildPatchTool, Display, TEXT("NB: If -MergeFileList is not specified, then union of all files will be selected, preferring ManifestB's version."));
-			UE_LOG(LogBuildPatchTool, Display, TEXT("NB: With the exception of the new version string for the build, all meta will be copied from only ManifestB."));
-			UE_LOG(LogBuildPatchTool, Display, TEXT(""));
+			PrintHelp<FMergeManifestToolModeHelp>();
 			return EReturnCode::OK;
 		}
 
+		LogConfig();
+
+		// Register secrets.
+		for (const TPair<FGuid, TArray<uint8>>& IdKeyPair : EncryptionSecrets)
+		{
+			BpsInterface.RegisterEncryptionSecret(IdKeyPair.Key, IdKeyPair.Value);
+		}
+
 		// Run the merge manifest routine
 		bool bSuccess = BpsInterface.MergeManifests(ManifestA, ManifestB, ManifestC, BuildVersion, MergeFileList);
 		return bSuccess ? EReturnCode::OK : EReturnCode::ToolFailure;
@@ -59,7 +55,7 @@ private:
 	{
 #define PARSE_SWITCH(Switch) ParseSwitch(TEXT(#Switch "="), Switch, Switches)
 		TArray<FString> Tokens, Switches;
-		FCommandLine::Parse(FCommandLine::Get(), Tokens, Switches);
+		FCommandLine::Parse(CommandLine, Tokens, Switches);
 
 		bHelp = ParseOption(TEXT("help"), Switches);
 		if (bHelp)
@@ -76,29 +72,64 @@ private:
 			UE_LOG(LogBuildPatchTool, Error, TEXT("ManifestA, ManifestB, ManifestC, and BuildVersion are required parameters"));
 			return false;
 		}
+
+		if (ManifestA.IsEmpty() || ManifestB.IsEmpty() || ManifestC.IsEmpty() || BuildVersion.IsEmpty())
+		{
+			UE_LOG(LogBuildPatchTool, Error, TEXT("ManifestA/B/C and BuildVersion arguments can't be empty"));
+			return false;
+		}
+
 		NormalizeUriFile(ManifestA);
 		NormalizeUriFile(ManifestB);
 		NormalizeUriFile(ManifestC);
 
+		if (ManifestC == ManifestA || ManifestC == ManifestB)
+		{
+			UE_LOG(LogBuildPatchTool, Error, TEXT("ManifestC parameter can't point to the same file as ManifestA or ManifestB"));
+			return false;
+		}
+
 		// Optional list to pick specific files, otherwise it is A stomped by B
 		PARSE_SWITCH(MergeFileList);
 		NormalizeUriFile(MergeFileList);
 
+		// Parse all the secret key and ID pairs.
+		FString ParseErrorMessage;
+		if (!ParsePairs(TEXT("EncryptionSecretId="), TEXT("EncryptionSecretKey="), Switches, EncryptionSecrets, ParseErrorMessage) && !ParseErrorMessage.IsEmpty())
+		{
+			UE_LOG(LogBuildPatchTool, Error, TEXT("%s"), *ParseErrorMessage);
+			return false;
+		}
+
 		return true;
 #undef PARSE_SWITCH
 	}
 
+	void LogConfig() const
+	{
+		UE_LOG(LogBuildPatchTool, Log, TEXT("-----Configuration for MergeManifestMode------"));
+		UE_LOG(LogBuildPatchTool, Log, TEXT("   ManifestA: %s"), *ManifestA);
+		UE_LOG(LogBuildPatchTool, Log, TEXT("   ManifestB: %s"), *ManifestB);
+		UE_LOG(LogBuildPatchTool, Log, TEXT("   ManifestC: %s"), *ManifestC);
+		UE_LOG(LogBuildPatchTool, Log, TEXT("   BuildVersion: %s"), *BuildVersion);
+		UE_LOG(LogBuildPatchTool, Log, TEXT("   MergeFileList: %s"), *MergeFileList);
+		for (const TPair<FGuid, TArray<uint8>>& EncryptionSecret : EncryptionSecrets)
+		{
+			UE_LOG(LogBuildPatchTool, Log, TEXT("   AvailableSecretId: %s"), *EncryptionSecret.Key.ToString());
+		}
+	}
+
 private:
+	static TCHAR const* const MODE_NAME;
 	IBuildPatchServicesModule& BpsInterface;
+	const TCHAR* CommandLine;
 	bool bHelp;
 	FString ManifestA;
 	FString ManifestB;
 	FString ManifestC;
 	FString BuildVersion;
 	FString MergeFileList;
+	TMap<FGuid, TArray<uint8>> EncryptionSecrets;
 };
 
-BuildPatchTool::IToolModeRef BuildPatchTool::FMergeManifestToolModeFactory::Create(IBuildPatchServicesModule& BpsInterface)
-{
-	return MakeShareable(new FMergeManifestToolMode(BpsInterface));
-}
+IMPLEMENT_BPT_MODE(MergeManifests, FMergeManifestToolMode);
diff --git a/Engine/Source/Programs/BuildPatchTool/Private/ToolModes/PackageChunksMode.cpp b/Engine/Source/Programs/BuildPatchTool/Private/ToolModes/PackageChunksMode.cpp
--- a/Engine/Source/Programs/BuildPatchTool/Private/ToolModes/PackageChunksMode.cpp
+++ b/Engine/Source/Programs/BuildPatchTool/Private/ToolModes/PackageChunksMode.cpp
@@ -1,19 +1,23 @@
 // Copyright Epic Games, Inc. All Rights Reserved.
 
 #include "ToolModes/PackageChunksMode.h"
-#include "Interfaces/IBuildPatchServicesModule.h"
+
+#include "Algo/Find.h"
+#include "Algo/Count.h"
 #include "BuildPatchTool.h"
+#include "Interfaces/IBuildPatchServicesModule.h"
 #include "Misc/CommandLine.h"
 #include "Misc/Paths.h"
-#include "Algo/Find.h"
+#include "ToolModes/ToolModesHelp.h"
 
 using namespace BuildPatchTool;
 
 class FPackageChunksToolMode : public IToolMode
 {
 public:
-	FPackageChunksToolMode(IBuildPatchServicesModule& InBpsInterface)
+	FPackageChunksToolMode(IBuildPatchServicesModule& InBpsInterface, const TCHAR* InCommandLine)
 		: BpsInterface(InBpsInterface)
+		, CommandLine(InCommandLine)
 	{}
 
 	virtual ~FPackageChunksToolMode()
@@ -30,48 +34,22 @@ public:
 		// Print help if requested.
 		if (bHelp)
 		{
-			UE_LOG(LogBuildPatchTool, Display, TEXT("PACKAGE CHUNKS MODE"));
-			UE_LOG(LogBuildPatchTool, Display, TEXT("This tool mode supports packaging data required for an installation into larger files which can be used as local sources for build patch installers."));
-			UE_LOG(LogBuildPatchTool, Display, TEXT(""));
-			UE_LOG(LogBuildPatchTool, Display, TEXT("Required arguments:"));
-			UE_LOG(LogBuildPatchTool, Display, TEXT("  -mode=PackageChunks  Must be specified to launch the tool in package chunks mode."));
-			UE_LOG(LogBuildPatchTool, Display, TEXT("  -FeatureLevel=Latest Specifies the client feature level to output data for. See BuildPatchServices::EFeatureLevel for possible values."));
-			UE_LOG(LogBuildPatchTool, Display, TEXT("  -ManifestFile=\"\"     Specifies in quotes the file path to the manifest to enumerate chunks from."));
-			UE_LOG(LogBuildPatchTool, Display, TEXT("  -OutputFile=\"\"       Specifies in quotes the file path the output package. Extension of .chunkdb will be added if not present."));
-			UE_LOG(LogBuildPatchTool, Display, TEXT(""));
-			UE_LOG(LogBuildPatchTool, Display, TEXT("Optional arguments:"));
-			UE_LOG(LogBuildPatchTool, Display, TEXT("  -PrevManifestFile=\"\" Specifies in quotes the file path to a manifest for a previous build, this will be used to filter out chunks, such that the"));
-			UE_LOG(LogBuildPatchTool, Display, TEXT("                       produced chunkdb files will only contain chunks required to patch from this build to the one described by ManifestFile."));
-			UE_LOG(LogBuildPatchTool, Display, TEXT("  -CloudDir=\"\"         Specifies in quotes the cloud directory where chunks to be packaged can be found."));
-			UE_LOG(LogBuildPatchTool, Display, TEXT("  -MaxOutputFileSize=  When specified, the size of each output file (in bytes) will be limited to a maximum of the provided value."));
-			UE_LOG(LogBuildPatchTool, Display, TEXT("  -ResultDataFile=\"\"   Specifies in quotes the file path where the results will be exported as a JSON object."));
-			UE_LOG(LogBuildPatchTool, Display, TEXT("  -TagSets=\",t1,t2\"     Specifies in quotes a comma seperated tagset for filtering of data saved. Multiple sets can also be provided to split the chunkdb files by tagsets."));
-			UE_LOG(LogBuildPatchTool, Display, TEXT("                        Untagged files will be referenced with an empty tag, which you can specify using an extra comma."));
-			UE_LOG(LogBuildPatchTool, Display, TEXT("  -PrevTagSet=\",tA,tB\"  Specifies in quotes a comma seperated tagset for filtering of input data usable from PrevManifestFile. This will increase the amount of chunks"));
-			UE_LOG(LogBuildPatchTool, Display, TEXT("                        saved out by reducing the number of files from the input manifest that are assumed usable. Only one PrevTagSet should be provided."));
-			UE_LOG(LogBuildPatchTool, Display, TEXT("                        Untagged files will be referenced with an empty tag, which you can specify using an extra comma."));
-			UE_LOG(LogBuildPatchTool, Display, TEXT(""));
-			UE_LOG(LogBuildPatchTool, Display, TEXT("NB: If CloudDir is not specified, the manifest file location will be used as the cloud directory."));
-			UE_LOG(LogBuildPatchTool, Display, TEXT("NB: If an optimised delta was available, the file extension .delta.chunkdb will be used."));
-			UE_LOG(LogBuildPatchTool, Display, TEXT("NB: MaxOutputFileSize is recommended to be as large as possible. The minimum individual chunkdb filesize is equal to one chunk plus chunkdb"));
-			UE_LOG(LogBuildPatchTool, Display, TEXT("    header, and thus will not result in efficient behavior."));
-			UE_LOG(LogBuildPatchTool, Display, TEXT("NB: If MaxOutputFileSize is not specified, the one output file will be produced containing all required data."));
-			UE_LOG(LogBuildPatchTool, Display, TEXT("NB: If MaxOutputFileSize is specified, the output files will be generated as Name.part01.chunkdb, Name.part02.chunkdb etc. The part number will"));
-			UE_LOG(LogBuildPatchTool, Display, TEXT("    have the number of digits required for highest numbered part."));
-			UE_LOG(LogBuildPatchTool, Display, TEXT("NB: If MaxOutputFileSize is specified, then each part can be equal to or less than the specified size, depending on the size of the last chunk"));
-			UE_LOG(LogBuildPatchTool, Display, TEXT("    that fits."));
-			UE_LOG(LogBuildPatchTool, Display, TEXT("NB: When providing multiple -TagSets= arguments, all data from the first -TagSets= arg will be saved first, followed by any extra data needed for the second -TagSets= arg, and so on in separated chunkdb files."));
-			UE_LOG(LogBuildPatchTool, Display, TEXT("    Note that this means the chunkdb files produced for the second -TagSets= arg and later will not contain some required data for that tagset if the data already got saved out as part of a previous tagset."));
-			UE_LOG(LogBuildPatchTool, Display, TEXT("    The chunkdb files are thus additive with no dupes."));
-			UE_LOG(LogBuildPatchTool, Display, TEXT("    If it is desired that each tagset's chunkdb files contain the duplicate data, then PackageChunks should be executed once per -TagSets= arg rather than once will all -TagSets= args."));
-			UE_LOG(LogBuildPatchTool, Display, TEXT("    An empty tag must be included in one of the -TagSets= args to include untagged file data in that tagset, e.g. -TagSets=\" , t1\"."));
-			UE_LOG(LogBuildPatchTool, Display, TEXT("    Adding no -TagSets= args will include all data."));
+			PrintHelp<FPackageChunksToolModeHelp>();
 			return EReturnCode::OK;
 		}
 
+		LogConfig();
+
+		// Register secrets.
+		for (const TPair<FGuid, TArray<uint8>>& IdKeyPair : EncryptionSecrets)
+		{
+			BpsInterface.RegisterEncryptionSecret(IdKeyPair.Key, IdKeyPair.Value);
+		}
+
 		// Setup and run.
 		BuildPatchServices::FPackageChunksConfiguration Configuration;
-		if (!BuildPatchServices::FeatureLevelFromString(*FeatureLevel, Configuration.FeatureLevel))
+		LexFromString(Configuration.FeatureLevel, *FeatureLevel);
+		if (Configuration.FeatureLevel == BuildPatchServices::EFeatureLevel::Invalid)
 		{
 			UE_LOG(LogBuildPatchTool, Error, TEXT("Provided FeatureLevel is not recognised. Invalid arg: -FeatureLevel=%s"), *FeatureLevel);
 			return EReturnCode::ArgumentProcessingError;
@@ -84,6 +62,9 @@ public:
 		Configuration.CloudDir = CloudDir;
 		Configuration.MaxOutputFileSize = MaxOutputFileSize;
 		Configuration.ResultDataFilePath = ResultDataFile;
+		Configuration.bIgnoreManifestFileInvalidTags = bIgnoreManifestFileInvalidTags;
+		Configuration.bIgnorePrevManifestFileInvalidTags = bIgnorePrevManifestFileInvalidTags;
+		Configuration.DeltaFilenameTrailer = DeltaFilenameTrailer;
 
 		// Run the enumeration routine.
 		bool bSuccess = BpsInterface.PackageChunkData(Configuration);
@@ -95,10 +76,11 @@ private:
 	bool ProcessCommandline()
 	{
 #define HAS_SWITCH(SwitchVar) (Algo::FindByPredicate(Switches, [](const FString& Elem){ return Elem.StartsWith(TEXT(#SwitchVar "="));}) != nullptr)
+#define HAS_DUP_SWITCH(SwitchVar) (Algo::CountIf(Switches, [](const FString& Elem){ return Elem.StartsWith(TEXT(#SwitchVar "="));}) > 1)
 #define PARSE_SWITCH(SwitchVar) ParseSwitch(TEXT(#SwitchVar "="), SwitchVar, Switches)
 #define PARSE_SWITCHES(SwitchesVar) ParseSwitches(TEXT(#SwitchesVar "="), SwitchesVar, Switches)
 		TArray<FString> Tokens, Switches;
-		FCommandLine::Parse(FCommandLine::Get(), Tokens, Switches);
+		FCommandLine::Parse(CommandLine, Tokens, Switches);
 
 		bHelp = ParseOption(TEXT("help"), Switches);
 		if (bHelp)
@@ -122,6 +104,19 @@ private:
 			UE_LOG(LogBuildPatchTool, Error, TEXT("ManifestFile and OutputFile are required parameters"));
 			return false;
 		}
+
+		if (ManifestFile.IsEmpty())
+		{
+			UE_LOG(LogBuildPatchTool, Error, TEXT("ManifestFile argument can't be empty"));
+			return false;
+		}
+
+		if (OutputFile.IsEmpty())
+		{
+			UE_LOG(LogBuildPatchTool, Error, TEXT("OutputFile argument can't be empty"));
+			return false;
+		}
+
 		NormalizeUriFile(ManifestFile);
 		NormalizeUriFile(OutputFile);
 
@@ -130,8 +125,22 @@ private:
 		PARSE_SWITCH(ResultDataFile);
 		NormalizeUriFile(PrevManifestFile);
 		NormalizeUriFile(ResultDataFile);
+		TArray<FString> TagSets;
 		PARSE_SWITCHES(TagSets);
 		PARSE_SWITCH(PrevTagSet);
+		PARSE_SWITCH(DeltaFilenameTrailer);
+
+		if (HAS_DUP_SWITCH(PrevTagSet))
+		{
+			UE_LOG(LogBuildPatchTool, Error, TEXT("Only one PrevTagSet should be provided"));
+			return false;
+		}
+
+		if (ManifestFile == PrevManifestFile)
+		{
+			UE_LOG(LogBuildPatchTool, Error, TEXT("ManifestFile and PreManifestFile should not point to the same file"));
+			return false;
+		}
 
 		if (!PARSE_SWITCH(CloudDir))
 		{
@@ -148,6 +157,13 @@ private:
 				UE_LOG(LogBuildPatchTool, Error, TEXT("MaxOutputFileSize must be a valid uint64"));
 				return false;
 			}
+
+			if (MaxOutputFileSize < 10000000)
+			{
+				// MaxOutputFileSize must be at least 10MB, otherwise fail.
+				UE_LOG(LogBuildPatchTool, Error, TEXT("MaxOutputFileSize must be at least 10000000 (10MB)"));
+				return false;
+			}
 		}
 		else
 		{
@@ -176,28 +192,72 @@ private:
 			}
 		}
 
+		bIgnoreManifestFileInvalidTags = ParseOption(TEXT("IgnoreManifestFileInvalidTags"), Switches);
+		bIgnorePrevManifestFileInvalidTags = ParseOption(TEXT("IgnorePrevManifestFileInvalidTags"), Switches);
+
+		// Parse all the secret key and ID pairs.
+		FString ParseErrorMessage;
+		if (!ParsePairs(TEXT("EncryptionSecretId="), TEXT("EncryptionSecretKey="), Switches, EncryptionSecrets, ParseErrorMessage) && !ParseErrorMessage.IsEmpty())
+		{
+			UE_LOG(LogBuildPatchTool, Error, TEXT("%s"), *ParseErrorMessage);
+			return false;
+		}
+
 		return true;
 #undef PARSE_SWITCHES
 #undef PARSE_SWITCH
+#undef HAS_DUP_SWITCH
 #undef HAS_SWITCH
 	}
 
+	void LogConfig() const
+	{
+		UE_LOG(LogBuildPatchTool, Log, TEXT("-----Configuration for PackageChunksMode------"));
+		UE_LOG(LogBuildPatchTool, Log, TEXT("   ManifestFile: %s"), *ManifestFile);
+		UE_LOG(LogBuildPatchTool, Log, TEXT("   PrevManifestFile: %s"), *PrevManifestFile);
+		UE_LOG(LogBuildPatchTool, Log, TEXT("   OutputFile: %s"), *OutputFile);
+		UE_LOG(LogBuildPatchTool, Log, TEXT("   ResultDataFile: %s"), *ResultDataFile);
+		UE_LOG(LogBuildPatchTool, Log, TEXT("   CloudDir: %s"), *CloudDir);
+		UE_LOG(LogBuildPatchTool, Log, TEXT("   MaxOutputFileSize: %llu"), MaxOutputFileSize);
+		for (const TSet<FString>& TagSet : TagSetArray)
+		{
+			FString TagSetString = FString::Join(TagSet.Array(), TEXT(","));
+			UE_LOG(LogBuildPatchTool, Log, TEXT("   With Tag set: %s"), *TagSetString);
+		}
+		FString PrevTagSetString = FString::Join(PrevTagSet.Array(), TEXT(","));
+		UE_LOG(LogBuildPatchTool, Log, TEXT("   Previous Tag set: %s"), *PrevTagSetString);
+		if (bIgnoreManifestFileInvalidTags)
+		{
+			UE_LOG(LogBuildPatchTool, Log, TEXT("   IgnoreManifestFileInvalidTags: Enabled"));
+		}
+		if (bIgnorePrevManifestFileInvalidTags)
+		{
+			UE_LOG(LogBuildPatchTool, Log, TEXT("   IgnorePrevManifestFileInvalidTags: Enabled"));
+		}
+		for (const TPair<FGuid, TArray<uint8>>& EncryptionSecret : EncryptionSecrets)
+		{
+			UE_LOG(LogBuildPatchTool, Log, TEXT("   AvailableSecretId: %s"), *EncryptionSecret.Key.ToString());
+		}
+	}
+
 private:
+	static TCHAR const* const MODE_NAME;
 	IBuildPatchServicesModule& BpsInterface;
+	const TCHAR* CommandLine;
 	bool bHelp;
 	FString FeatureLevel;
 	FString ManifestFile;
 	FString PrevManifestFile;
 	FString OutputFile;
 	FString ResultDataFile;
+	FString DeltaFilenameTrailer;
 	FString CloudDir;
 	uint64 MaxOutputFileSize;
-	TArray<FString> TagSets;
 	TArray<TSet<FString>> TagSetArray;
 	TSet<FString> PrevTagSet;
+	bool bIgnoreManifestFileInvalidTags;
+	bool bIgnorePrevManifestFileInvalidTags;
+	TMap<FGuid, TArray<uint8>> EncryptionSecrets;
 };
 
-BuildPatchTool::IToolModeRef BuildPatchTool::FPackageChunksToolModeFactory::Create(IBuildPatchServicesModule& BpsInterface)
-{
-	return MakeShareable(new FPackageChunksToolMode(BpsInterface));
-}
+IMPLEMENT_BPT_MODE(PackageChunks, FPackageChunksToolMode);
diff --git a/Engine/Source/Programs/BuildPatchTool/Private/ToolModes/VerifyChunksMode.cpp b/Engine/Source/Programs/BuildPatchTool/Private/ToolModes/VerifyChunksMode.cpp
--- a/Engine/Source/Programs/BuildPatchTool/Private/ToolModes/VerifyChunksMode.cpp
+++ b/Engine/Source/Programs/BuildPatchTool/Private/ToolModes/VerifyChunksMode.cpp
@@ -1,19 +1,22 @@
 // Copyright Epic Games, Inc. All Rights Reserved.
 
 #include "ToolModes/VerifyChunksMode.h"
-#include "Interfaces/IBuildPatchServicesModule.h"
+
+#include "Algo/Find.h"
 #include "BuildPatchTool.h"
+#include "Interfaces/IBuildPatchServicesModule.h"
 #include "Misc/CommandLine.h"
 #include "Misc/Paths.h"
-#include "Algo/Find.h"
+#include "ToolModes/ToolModesHelp.h"
 
 using namespace BuildPatchTool;
 
 class FVerifyChunksToolMode : public IToolMode
 {
 public:
-	FVerifyChunksToolMode(IBuildPatchServicesModule& InBpsInterface)
+	FVerifyChunksToolMode(IBuildPatchServicesModule& InBpsInterface, const TCHAR* InCommandLine)
 		: BpsInterface(InBpsInterface)
+		, CommandLine(InCommandLine)
 	{}
 
 	virtual ~FVerifyChunksToolMode()
@@ -30,20 +33,18 @@ public:
 		// Print help if requested.
 		if (bHelp)
 		{
-			UE_LOG(LogBuildPatchTool, Display, TEXT("VERIFY CHUNKS MODE"));
-			UE_LOG(LogBuildPatchTool, Display, TEXT("This tool mode allows you to verify the integrity of patch data. It will load chunk or chunkdb files to verify they are not corrupt."));
-			UE_LOG(LogBuildPatchTool, Display, TEXT(""));
-			UE_LOG(LogBuildPatchTool, Display, TEXT("Required arguments:"));
-			UE_LOG(LogBuildPatchTool, Display, TEXT("  -mode=VerifyChunks  Must be specified to launch the tool in verify chunks mode."));
-			UE_LOG(LogBuildPatchTool, Display, TEXT("  -SearchPath=\"\"      Specifies in quotes the directory path which contains data to verify."));
-			UE_LOG(LogBuildPatchTool, Display, TEXT(""));
-			UE_LOG(LogBuildPatchTool, Display, TEXT("Optional arguments:"));
-			UE_LOG(LogBuildPatchTool, Display, TEXT("  -OutputFile=\"\"      When specified, full file path for each bad data will be saved to this file as \\r\\n seperated list."));
-			UE_LOG(LogBuildPatchTool, Display, TEXT(""));
-			UE_LOG(LogBuildPatchTool, Display, TEXT("NB: All checks are logged, normal log for good data, error log for any bad data found."));
+			PrintHelp<FVerifyChunksToolModeHelp>();
 			return EReturnCode::OK;
 		}
 
+		LogConfig();
+
+		// Register secrets.
+		for (const TPair<FGuid, TArray<uint8>>& IdKeyPair : EncryptionSecrets)
+		{
+			BpsInterface.RegisterEncryptionSecret(IdKeyPair.Key, IdKeyPair.Value);
+		}
+
 		// Run the enumeration routine.
 		bool bSuccess = BpsInterface.VerifyChunkData(SearchPath, OutputFile);
 		return bSuccess ? EReturnCode::OK : EReturnCode::ToolFailure;
@@ -56,7 +57,7 @@ private:
 #define HAS_SWITCH(Switch) (Algo::FindByPredicate(Switches, [](const FString& Elem){ return Elem.StartsWith(TEXT(#Switch "="));}) != nullptr)
 #define PARSE_SWITCH(Switch) ParseSwitch(TEXT(#Switch "="), Switch, Switches)
 		TArray<FString> Tokens, Switches;
-		FCommandLine::Parse(FCommandLine::Get(), Tokens, Switches);
+		FCommandLine::Parse(CommandLine, Tokens, Switches);
 
 		bHelp = ParseOption(TEXT("help"), Switches);
 		if (bHelp)
@@ -70,25 +71,51 @@ private:
 			UE_LOG(LogBuildPatchTool, Error, TEXT("SearchPath is a required parameter"));
 			return false;
 		}
+
+		if (SearchPath.IsEmpty())
+		{
+			UE_LOG(LogBuildPatchTool, Error, TEXT("SearchPath argument can't be empty"));
+			return false;
+		}
+
 		NormalizeUriPath(SearchPath);
 
 		// Get optional parameters.
 		PARSE_SWITCH(OutputFile);
 		NormalizeUriFile(OutputFile);
 
+		// Parse all the secret key and ID pairs.
+		FString ParseErrorMessage;
+		if (!ParsePairs(TEXT("EncryptionSecretId="), TEXT("EncryptionSecretKey="), Switches, EncryptionSecrets, ParseErrorMessage) && !ParseErrorMessage.IsEmpty())
+		{
+			UE_LOG(LogBuildPatchTool, Error, TEXT("%s"), *ParseErrorMessage);
+			return false;
+		}
+
 		return true;
 #undef PARSE_SWITCH
 #undef HAS_SWITCH
 	}
 
+	void LogConfig() const
+	{
+		UE_LOG(LogBuildPatchTool, Log, TEXT("-----Configuration for VerifyChunksMode------"));
+		UE_LOG(LogBuildPatchTool, Log, TEXT("   SearchPath: %s"), *SearchPath);
+		UE_LOG(LogBuildPatchTool, Log, TEXT("   OutputFile: %s"), *OutputFile);
+		for (const TPair<FGuid, TArray<uint8>>& EncryptionSecret : EncryptionSecrets)
+		{
+			UE_LOG(LogBuildPatchTool, Log, TEXT("   AvailableSecretId: %s"), *EncryptionSecret.Key.ToString());
+		}
+	}
+
 private:
+	static TCHAR const* const MODE_NAME;
 	IBuildPatchServicesModule& BpsInterface;
+	const TCHAR* CommandLine;
 	bool bHelp;
 	FString SearchPath;
 	FString OutputFile;
+	TMap<FGuid, TArray<uint8>> EncryptionSecrets;
 };
 
-BuildPatchTool::IToolModeRef BuildPatchTool::FVerifyChunksToolModeFactory::Create(IBuildPatchServicesModule& BpsInterface)
-{
-	return MakeShareable(new FVerifyChunksToolMode(BpsInterface));
-}
+IMPLEMENT_BPT_MODE(VerifyChunks, FVerifyChunksToolMode);
diff --git a/Engine/Source/Programs/BuildPatchTool/Public/BuildPatchTool.h b/Engine/Source/Programs/BuildPatchTool/Public/BuildPatchTool.h
--- a/Engine/Source/Programs/BuildPatchTool/Public/BuildPatchTool.h
+++ b/Engine/Source/Programs/BuildPatchTool/Public/BuildPatchTool.h
@@ -15,7 +15,47 @@ namespace BuildPatchTool
 		UnknownToolMode,
 		FileNotFound,
 		ToolFailure,
+		OnlineMcpError,
+		DeprecatedVersion,
+		DirectoryNotFound,
+		MissingCredentials,
+		InvalidCredentials,
+		InvalidData,
+
+		// !! Always after the last code, signifies the last return code plus 1 to allow the following Last alias.
+		LastPlusOne,
+		// An alias for the actual last return code.
+		Last = (LastPlusOne - 1),
 
 		Crash = 255
 	};
-}
\ No newline at end of file
+}
+
+namespace CommandLineHelpers
+{
+	bool ParseSwitch(const TCHAR* InSwitch, FString& Value, const TArray<FString>& Switches);
+}
+
+inline const TCHAR* LexToString(BuildPatchTool::EReturnCode ReturnCode)
+{
+	static_assert(BuildPatchTool::EReturnCode::Last == BuildPatchTool::EReturnCode::InvalidData, "Please add support for the extra values below.");
+#define CASE_ENUM_TO_STR(Value) case BuildPatchTool::EReturnCode::Value: return TEXT(#Value)
+	switch (ReturnCode)
+	{
+		CASE_ENUM_TO_STR(OK);
+		CASE_ENUM_TO_STR(UnknownError);
+		CASE_ENUM_TO_STR(ArgumentProcessingError);
+		CASE_ENUM_TO_STR(UnknownToolMode);
+		CASE_ENUM_TO_STR(FileNotFound);
+		CASE_ENUM_TO_STR(ToolFailure);
+		CASE_ENUM_TO_STR(OnlineMcpError);
+		CASE_ENUM_TO_STR(DeprecatedVersion);
+		CASE_ENUM_TO_STR(DirectoryNotFound);
+		CASE_ENUM_TO_STR(MissingCredentials);
+		CASE_ENUM_TO_STR(InvalidCredentials);
+		CASE_ENUM_TO_STR(InvalidData);
+		CASE_ENUM_TO_STR(Crash);
+	default: return TEXT("InvalidOrMax");
+	}
+#undef CASE_ENUM_TO_STR
+}
diff --git a/Engine/Source/Programs/BuildPatchTool/Public/Common/CmdLineProcessor.h b/Engine/Source/Programs/BuildPatchTool/Public/Common/CmdLineProcessor.h
new file mode 100644
--- /dev/null
+++ b/Engine/Source/Programs/BuildPatchTool/Public/Common/CmdLineProcessor.h
@@ -0,0 +1,342 @@
+// Copyright Epic Games, Inc. All Rights Reserved.
+#pragma once
+
+#include "Containers/UnrealString.h"
+#include "Templates/Function.h"
+#include "Misc/Paths.h"
+#include "Common/DiffAbortThresholdArgument.h"
+
+#define BPT_DECLARE_COMMANDLINE(TConfig) \
+bool ProcessCommandLine(); \
+static TCHAR const * const MODE_NAME; \
+TConfig Config;
+
+#define BPT_COMMANDLINE_START(TClass) \
+bool TClass::ProcessCommandLine() \
+{ \
+	bool bParamExists = true; \
+	bool bSuccess = true; \
+	TArray<FString> Tokens, Switches; \
+	FCommandLine::Parse(CommandLine, Tokens, Switches); \
+	bHelp = ParseOption(TEXT("help"), Switches); \
+	if (bHelp) \
+	{ \
+		return true; \
+	}
+
+#define BPT_ARG_REQUIRED_SWITCH(Switch) \
+	bParamExists = ParseSwitch(TEXT(#Switch "="), Config.Switch, Switches); \
+	if (!bParamExists) \
+	{ \
+		UE_LOG(LogBuildPatchTool, Error, TEXT("%s"), *FString::Format(TEXT("{0} is required for {1} mode."), { TEXT(#Switch), MODE_NAME })); \
+		bSuccess = false; \
+	}
+
+#define BPT_ARG_REQUIRED_SWITCH_PAIRS(Key, Val, Map) \
+	{ \
+		FString ParseErrorMessage; \
+		bParamExists = ParsePairs(TEXT(#Key "="), TEXT(#Val "="), Switches, Config.Map, ParseErrorMessage); \
+		if (!ParseErrorMessage.IsEmpty()) \
+		{ \
+			UE_LOG(LogBuildPatchTool, Error, TEXT("%s"), *ParseErrorMessage); \
+			bParamExists = true; \
+			bSuccess = false; \
+		} \
+		if (!bParamExists) \
+		{ \
+			UE_LOG(LogBuildPatchTool, Error, TEXT("%s"), *FString::Format(TEXT("Paired {0} and {1} are required for {2} mode."), { TEXT(#Key), TEXT(#Val), MODE_NAME })); \
+			bSuccess = false; \
+		} \
+	}
+
+#define BPT_ARG_OPTIONAL_SWITCH_PAIRS(Key, Val, Map) \
+	{ \
+		FString ParseErrorMessage; \
+		bParamExists = ParsePairs(TEXT(#Key "="), TEXT(#Val "="), Switches, Config.Map, ParseErrorMessage); \
+		if (!ParseErrorMessage.IsEmpty()) \
+		{ \
+			UE_LOG(LogBuildPatchTool, Error, TEXT("%s"), *ParseErrorMessage); \
+			bParamExists = true; \
+			bSuccess = false; \
+		} \
+	}
+
+#define BPT_ARG_REQUIRED_SWITCH_NOEMPTY_ALIAS(Switch, VarName) \
+	bParamExists = ParseSwitch(TEXT(#Switch "="), Config.VarName, Switches); \
+	if (!bParamExists) \
+	{ \
+		UE_LOG(LogBuildPatchTool, Error, TEXT("%s"), *FString::Format(TEXT("{0} is required for {1} mode."), { TEXT(#Switch), MODE_NAME })); \
+		bSuccess = false; \
+	} \
+	else if (FCmdLineProcessor::IsEmpty(Config.VarName)) \
+	{ \
+		UE_LOG(LogBuildPatchTool, Error, TEXT("%s"), *FString::Format(TEXT("{0} should not be provided empty."), { TEXT(#Switch) })); \
+		bSuccess = false; \
+	}
+
+#define BPT_ARG_REQUIRED_SWITCH_NOEMPTY(Switch) BPT_ARG_REQUIRED_SWITCH_NOEMPTY_ALIAS(Switch, Switch)
+
+#define BPT_ARG_REQUIRED_FILENAME(Switch) \
+	bParamExists = ParseSwitch(TEXT(#Switch "="), Config.Switch, Switches); \
+	if (!bParamExists) \
+	{ \
+		UE_LOG(LogBuildPatchTool, Error, TEXT("%s"), *FString::Format(TEXT("{0} is required for {1} mode."), { TEXT(#Switch), MODE_NAME })); \
+		bSuccess = false; \
+	} \
+	else \
+	{ \
+		Config.Switch.TrimStartAndEndInline(); \
+		if (FCmdLineProcessor::IsEmpty(Config.Switch)) \
+		{ \
+			UE_LOG(LogBuildPatchTool, Error, TEXT("%s"), *FString::Format(TEXT("{0} should not be provided empty."), { TEXT(#Switch) })); \
+			bSuccess = false; \
+		} \
+		NormalizeUriFile(Config.Switch); \
+		if (Config.Switch.EndsWith(TEXT("/"))) \
+		{ \
+			UE_LOG(LogBuildPatchTool, Error, TEXT("%s"), *FString::Format(TEXT("{0} should be a path to a file, not a directory."), { TEXT(#Switch) })); \
+			bSuccess = false; \
+		} \
+	}
+
+#define BPT_ARG_REQUIRED_FILENAME_EXISTS(Switch) \
+	BPT_ARG_REQUIRED_FILENAME(Switch) \
+	if (bParamExists && !FPaths::FileExists(Config.Switch)) \
+	{ \
+		UE_LOG(LogBuildPatchTool, Error, TEXT("%s"), *FString::Format(TEXT("{0} does not exist. [{1}]"), { TEXT(#Switch), *Config.Switch })); \
+		bSuccess = false; \
+	}
+
+#define BPT_ARG_OPTIONAL_FILENAME(Switch) \
+	bParamExists = ParseSwitch(TEXT(#Switch "="), Config.Switch, Switches); \
+	if (bParamExists && !FCmdLineProcessor::IsEmpty(Config.Switch)) \
+	{ \
+		if (Config.Switch.EndsWith(TEXT("/"))) \
+		{ \
+			UE_LOG(LogBuildPatchTool, Error, TEXT("%s"), *FString::Format(TEXT("{0} should be a path to a file, not a directory."), { TEXT(#Switch) })); \
+			bSuccess = false; \
+		} \
+		else \
+		{ \
+			NormalizeUriFile(Config.Switch); \
+		} \
+	}
+	
+#define BPT_ARG_OPTIONAL_FILENAME_EXISTS(Switch) \
+	BPT_ARG_OPTIONAL_FILENAME(Switch) \
+	if (bParamExists && !FPaths::FileExists(Config.Switch)) \
+	{ \
+		UE_LOG(LogBuildPatchTool, Error, TEXT("%s"), *FString::Format(TEXT("{0} does not exist. [{1}]"), { TEXT(#Switch), *Config.Switch })); \
+		bSuccess = false; \
+	}
+
+#define BPT_ARG_REQUIRED_PATH(Switch) \
+	bParamExists = ParseSwitch(TEXT(#Switch "="), Config.Switch, Switches); \
+	if (!bParamExists) \
+	{ \
+		UE_LOG(LogBuildPatchTool, Error, TEXT("%s"), *FString::Format(TEXT("{0} is required for {1} mode."), { TEXT(#Switch), MODE_NAME })); \
+		bSuccess = false; \
+	} \
+	else \
+	{ \
+		Config.Switch.TrimStartAndEndInline(); \
+		if (FCmdLineProcessor::IsEmpty(Config.Switch)) \
+		{ \
+			UE_LOG(LogBuildPatchTool, Error, TEXT("%s"), *FString::Format(TEXT("{0} should not be provided empty."), { TEXT(#Switch) })); \
+			bSuccess = false; \
+		} \
+		else \
+		{ \
+			NormalizeUriPath(Config.Switch); \
+			Config.Switch = FPaths::ConvertRelativePathToFull(Config.Switch); \
+		}\
+	}
+
+#define BPT_ARG_REQUIRED_PATH_NOEMPTY(Switch) \
+	BPT_ARG_REQUIRED_PATH(Switch) \
+	if (bParamExists && !FPaths::DirectoryExists(Config.Switch)) \
+	{ \
+		UE_LOG(LogBuildPatchTool, Error, TEXT("%s"), *FString::Format(TEXT("{0} does not exist. [{1}]"), { TEXT(#Switch), *Config.Switch })); \
+		bSuccess = false; \
+	} \
+	if (bParamExists && IsEmptyDirectory(Config.Switch)) \
+	{ \
+		UE_LOG(LogBuildPatchTool, Error, TEXT("%s"), *FString::Format(TEXT("{0} must not be an empty directory or a directory tree containing only empty files/directories. [{1}]"), { TEXT(#Switch), *Config.Switch })); \
+		bSuccess = false; \
+	}
+
+#define BPT_ARG_REQUIRED_SWITCH_MAXSIZE(Switch, MaxSize) \
+	bParamExists = ParseSwitch(TEXT(#Switch "="), Config.Switch, Switches); \
+	if (!bParamExists) \
+	{ \
+		UE_LOG(LogBuildPatchTool, Error, TEXT("%s"), *FString::Format(TEXT("{0} is required for {1} mode."), { TEXT(#Switch), MODE_NAME })); \
+		bSuccess = false; \
+	} \
+	else if (FCmdLineProcessor::IsEmpty(Config.Switch)) \
+	{ \
+		UE_LOG(LogBuildPatchTool, Error, TEXT("%s"), *FString::Format(TEXT("{0} should not be provided empty."), { TEXT(#Switch) })); \
+		bSuccess = false; \
+	} \
+	else if (FCmdLineProcessor::ExceedsSize(Config.Switch, MaxSize)) \
+	{ \
+		UE_LOG(LogBuildPatchTool, Error, TEXT("%s"), *FString::Format(TEXT("{0} should not exceed length {1}."), { TEXT(#Switch), MaxSize })); \
+		bSuccess = false; \
+	}
+	
+#define BPT_ARG_UNEXPECTED_SWITCH(Switch) \
+	if (Switches.Contains(TEXT(#Switch))) \
+	{ \
+		UE_LOG(LogBuildPatchTool, Error, TEXT("%s"), *FString::Format(TEXT("{0} should not be used with {1} mode."), { TEXT(#Switch), MODE_NAME })); \
+		bSuccess = false; \
+	}
+
+#define BPT_ARG_OPTIONAL_SWITCH(Switch) \
+	bParamExists = ParseSwitch(TEXT(#Switch "="), Config.Switch, Switches);
+
+#define BPT_ARG_OPTIONAL_SWITCH_INT32(Switch) \
+	bParamExists = ParseSwitch<int32>(TEXT(#Switch "="), Config.Switch, Switches);
+
+#define BPT_ARG_OPTIONAL_SWITCH_NOEMPTY(Switch) \
+	bParamExists = ParseSwitch(TEXT(#Switch "="), Config.Switch, Switches); \
+	if (bParamExists && FCmdLineProcessor::IsEmpty(Config.Switch)) \
+	{ \
+		UE_LOG(LogBuildPatchTool, Error, TEXT("%s"), *FString::Format(TEXT("Optional arg {0} should not be provided empty. Remove if not needed."), { TEXT(#Switch) })); \
+		bSuccess = false; \
+	}
+
+#define BPT_ARG_OPTIONAL_SWITCH_CLAMPED(Switch, TInt, Min, Max) \
+	{ \
+		TInt Requested; \
+		bParamExists = ParseSwitch(TEXT(#Switch "="), Requested, Switches); \
+		if (bParamExists) \
+		{ \
+			Config.Switch = FMath::Clamp<TInt>(Requested, Min, Max); \
+			if (Requested != Config.Switch) \
+			{ \
+				UE_LOG(LogBuildPatchTool, Warning, TEXT("%s"), *FString::Format(TEXT("Requested -{0}={1} is outside of allowed range {2} >= n >= {3}. Please update your arg to be within range. Continuing with {4}."), { TEXT(#Switch), Requested, Max, Min, Config.Switch })); \
+			} \
+		} \
+	}
+
+#define BPT_ARG_OPTIONAL_SWITCH_MULTI(Switch) \
+	bParamExists = ParseSwitches(TEXT(#Switch "="), Config.Switch##s, Switches);
+
+#define BPT_ARG_OPTION(Option) \
+	Config.b##Option = ParseOption(TEXT(#Option), Switches);
+
+#define BPT_ARG_OPTION_MEMBER(Option) \
+	b##Option = ParseOption(TEXT(#Option), Switches);
+
+#define BPT_ARG_FEATURELEVEL_OPTIONAL_ALIAS(Switch) \
+	{ \
+		FString FeatureLevelString; \
+		bParamExists = ParseSwitch(TEXT(#Switch"="), FeatureLevelString, Switches); \
+		if (bParamExists) \
+		{ \
+			FeatureLevelString.TrimStartAndEndInline(); \
+			BuildPatchServices::EFeatureLevel FeatureLevel; \
+			LexFromString(FeatureLevel, *FeatureLevelString); \
+			if (BuildPatchServices::EFeatureLevel::Invalid == FeatureLevel) \
+			{ \
+				UE_LOG(LogBuildPatchTool, Error, TEXT("%s"), *FString::Format(TEXT("{0} FeatureLevel was provided, but {0} is not a recognized value."), { *FeatureLevelString })); \
+				bSuccess = false; \
+			} \
+			else \
+			{ \
+				Config.Switch = FeatureLevel; \
+			} \
+		} \
+	}
+#define BPT_ARG_FEATURELEVEL_OPTIONAL BPT_ARG_FEATURELEVEL_OPTIONAL_ALIAS(FeatureLevel) \
+
+#define BPT_ARG_FEATURELEVEL_REQUIRED \
+	{ \
+		FString FeatureLevelString; \
+		bParamExists = ParseSwitch(TEXT(#Switch"="), FeatureLevelString, Switches); \
+		if (bParamExists) \
+		{ \
+			FeatureLevelString.TrimStartAndEndInline(); \
+			BuildPatchServices::EFeatureLevel FeatureLevel; \
+			LexFromString(FeatureLevel, *FeatureLevelString); \
+			if (FCmdLineProcessor::IsEmpty(FeatureLevelString)) \
+			{ \
+				UE_LOG(LogBuildPatchTool, Error, TEXT("FeatureLevel should not be provided empty.")); \
+				bSuccess = false; \
+			} \
+			else if (BuildPatchServices::EFeatureLevel::Invalid == FeatureLevel) \
+			{ \
+				UE_LOG(LogBuildPatchTool, Error, TEXT("%s"), *FString::Format(TEXT("{0} FeatureLevel was provided, but {0} is not a recognized value."), { *FeatureLevelString })); \
+				bSuccess = false; \
+			} \
+			else \
+			{ \
+				Config.FeatureLevel = FeatureLevel; \
+			} \
+		} \
+		else \
+		{ \
+			UE_LOG(LogBuildPatchTool, Error, TEXT("%s"), *FString::Format(TEXT("FeatureLevel is required for {1} mode."), { TEXT(#Switch), MODE_NAME })); \
+			bSuccess = false; \
+		} \
+	}
+
+#define BPT_ARG_DIFFABORTTHRESHOLD_OPTIONAL \
+	Config.DiffAbortThreshold = BuildPatchTool::FDiffAbortThresholdArgument::Parse(*this, Switches);
+
+#define BPT_COMMANDLINE_FAIL_IF_TCHAR(Condition, ErrorStr) \
+	if (Condition) \
+	{ \
+		UE_LOG(LogBuildPatchTool, Error, TEXT("%s"), ErrorStr); \
+		bSuccess = false; \
+	}
+#define BPT_COMMANDLINE_FAIL_IF(Condition, ErrorStr) \
+	BPT_COMMANDLINE_FAIL_IF_TCHAR(Condition, TEXT(ErrorStr))
+
+#define BPT_COMMANDLINE_END \
+	return bSuccess; \
+}
+
+class FCmdLineProcessor
+{
+public:
+	FCmdLineProcessor(int32 ArgC, TCHAR* ArgV[], const FString& WorkingDirectory, TFunction<void(const FString&)> ErrorLogger);
+
+	bool ContainsFlag(const TCHAR* Flag) const;
+
+	bool ContainsSwitch(const TCHAR* Switch) const;
+
+	bool GetSwitchValue(const TCHAR* Switch, FString& OutValue) const;
+
+	const FString& GetOriginalCommandLine();
+
+	FString GetEngineString();
+
+	void ProcessDragAndDrop();
+
+	bool ProcessCommandlineFiles();
+
+	bool HandleLegacyCommandline();
+
+	bool ConvertCommandlinePaths();
+
+	void AddDesiredEngineFlags();
+
+	// Helpers for BPT_MODE_COMMANDLINE macros.
+	static inline bool IsEmpty(const FString& Val) { return Val.IsEmpty(); }
+	static inline bool ExceedsSize(const FString& Val, int32 Max) { return Val.Len() > Max; }
+
+private:
+	const FString WorkingDirectory;
+	TFunction<void(const FString&)> ErrorLogger;
+	FString OriginalCommandLine;
+	TArray<FString> Tokens;
+
+	const TArray<FString> PathParamsToMakeAbsolute = {
+		TEXT("-CloudDir="), TEXT("-ManifestA="), TEXT("-ManifestB="), TEXT("-ManifestC="), TEXT("-DownloadCloudDirectory="), TEXT("-ResultDataFile="),
+		TEXT("-BuildRoot="), TEXT("-CloudDirs="), TEXT("-OutputFile="), TEXT("-File="), TEXT("-MergeFileList="), TEXT("-PrevManifestFile="),
+		TEXT("-ManifestFile="), TEXT("-SearchPath="), TEXT("-InputFile="), TEXT("-Manifest="), TEXT("-OutputDir="),
+		TEXT("-DeletedChunkLogFile="), TEXT("-PrevManifest=") };
+
+	const TArray<FString> PathParamsToCollapseRelativeDirectories = {
+		TEXT("-AppLaunch=") };
+};
diff --git a/Engine/Source/Programs/BuildPatchTool/Public/Interfaces/ToolMode.h b/Engine/Source/Programs/BuildPatchTool/Public/Interfaces/ToolMode.h
--- a/Engine/Source/Programs/BuildPatchTool/Public/Interfaces/ToolMode.h
+++ b/Engine/Source/Programs/BuildPatchTool/Public/Interfaces/ToolMode.h
@@ -3,6 +3,64 @@
 
 #include "CoreMinimal.h"
 #include "Interfaces/IBuildPatchServicesModule.h"
+#include "Algo/AllOf.h"
+#include "HAL/FileManager.h"
+#include "BuildPatchTool.h"
+
+#define IMPLEMENT_BPT_MODE(ModeName, ModeClass) \
+class ModeClass; \
+TCHAR const * const ModeClass::MODE_NAME = TEXT(#ModeName); \
+struct FAutoRegister##ModeName \
+{ \
+	FAutoRegister##ModeName() \
+	{ \
+		using namespace BuildPatchTool; \
+		FToolModeFactory::GetModeRegistryDelegate().AddLambda([](TMap<FString, FToolModeFactory::FToolModeCreateFunc>& OutModes) \
+			{ \
+				OutModes.Add(TEXT(#ModeName), [](IBuildPatchServicesModule& BpsInterface, const TCHAR* CommandLine){ return new ModeClass(BpsInterface, CommandLine); }); \
+			}); \
+	} \
+} AutoRegister##ModeName;
+
+#define IMPLEMENT_BPT_MODE_ALIAS(ModeName, ModeClass) \
+struct FAutoRegister##ModeName \
+{ \
+	FAutoRegister##ModeName() \
+	{ \
+		using namespace BuildPatchTool; \
+		FToolModeFactory::GetModeRegistryDelegate().AddLambda([](TMap<FString, FToolModeFactory::FToolModeCreateFunc>& OutModes) \
+			{ \
+				OutModes.Add(TEXT(#ModeName), [](IBuildPatchServicesModule& BpsInterface, const TCHAR* CommandLine){ return new ModeClass(BpsInterface, CommandLine); }); \
+			}); \
+	} \
+} AutoRegister##ModeName;
+
+namespace Constants
+{
+	static const FString Comma(TEXT(","));
+	static const FString Equals(TEXT("="));
+	static const FString DoubleQuote(TEXT("\""));
+	static const FString SingleQuote(TEXT("'"));
+	static const FString Slash(TEXT("/"));
+	static const FString Backslash(TEXT("\\"));
+	static const FString Custom(TEXT("custom"));
+	static const FString CustomInt(TEXT("customint"));
+	static const FString CustomFloat(TEXT("customfloat"));
+}
+
+// Detect if a template parameter is a TSet
+template<typename T>
+struct is_ue_set : std::false_type {};
+
+template<typename... T>
+struct is_ue_set<TSet<T...>> : std::true_type {};
+
+// Detect if template parameter is a TArray
+template<typename T>
+struct is_ue_array : std::false_type {};
+
+template<typename... T>
+struct is_ue_array<TArray<T...>> : std::true_type {};
 
 namespace BuildPatchTool
 {
@@ -15,8 +73,20 @@ namespace BuildPatchTool
 		TCHAR const * const QuoteStr = TEXT("\"");
 
 	public:
+		virtual ~IToolMode() { }
+
 		virtual EReturnCode Execute() = 0;
 
+		bool bHelp = false;
+		template<typename T>
+		void PrintHelp()
+		{
+			for (int32 i = 0; i < T::NumLines; i++)
+			{
+				UE_LOG(LogBuildPatchTool, Display, TEXT("%s"), T::Text[i]);
+			}
+		}
+
 		/**
 		 * Helper for normalizing a URI referring to a file. We make sure not to cause issues with URIs that separate a protocol and authority with ://
 		 * and windows UNC paths that begin with \\. Windows UNC paths still work in UE4 code with //.
@@ -29,6 +99,14 @@ namespace BuildPatchTool
 			UriFile.ReplaceInline(TEXT("\\"), TEXT("/"), ESearchCase::CaseSensitive);
 		}
 
+		void NormalizeUriFileInKeys(TArray<TPair<FString, FString>>& Pairs)
+		{
+			for (TPair<FString, FString>& Pair : Pairs)
+			{
+				NormalizeUriFile(Pair.Key);
+			}
+		}
+
 		/**
 		 * Helper for normalizing a URI referring to a path. We make sure not to cause issues with URIs that separate a protocol and authority with ://
 		 * and windows UNC paths that begin with \\. Windows UNC paths still work in UE4 code with //.
@@ -49,6 +127,21 @@ namespace BuildPatchTool
 			}
 		}
 
+		bool IsEmptyDirectory(const FString& Directory)
+		{
+			IFileManager& FM = IFileManager::Get();
+			TArray<FString> FoundFiles;
+			FM.FindFilesRecursive(FoundFiles, *Directory, TEXT("*"), true, false);
+			for (const FString& File : FoundFiles)
+			{
+				if (FM.FileSize(*File) > 0)
+				{
+					return false;
+				}
+			}
+			return true;
+		}
+
 		/**
 		 * Helper for parsing a switch from an array of switches, usually produced using FCommandLine::Parse(..)
 		 *
@@ -64,14 +157,11 @@ namespace BuildPatchTool
 			checkSlow(InSwitch != nullptr);
 			checkSlow(InSwitch[FCString::Strlen(InSwitch)-1] == TEXT('='));
 
-			for (const FString& Switch : Switches)
+			if (const FString* Switch = Switches.FindByPredicate([InSwitch](const FString& Switch) { return Switch.StartsWith(InSwitch); }))
 			{
-				if (Switch.StartsWith(InSwitch))
-				{
-					FString StringValue;
-					Switch.Split(EqualsStr, nullptr, &StringValue);
-					return ParseValue(StringValue, Value);
-				}
+				FString StringValue;
+				Switch->Split(EqualsStr, nullptr, &StringValue);
+				return ParseValue(StringValue, Value);
 			}
 			return false;
 		}
@@ -114,6 +204,77 @@ namespace BuildPatchTool
 			return Switches.Contains(InSwitch);
 		}
 
+		template<typename TFirst, typename TSecond, typename TContainer>
+		bool ParsePairs(const TCHAR* InFirstSwitch, const TCHAR* InSecondSwitch, const TArray<FString>& Switches, TContainer& OutPairs, FString& ErrorMessage)
+		{
+			bool bFoundAnyPair = false;
+			bool bFoundFirst = false;
+			bool bFoundSecond = false;
+			TFirst FirstValue;
+			TSecond SecondValue;
+			for (int32 SwitchIndex = 0; SwitchIndex < Switches.Num(); SwitchIndex++)
+			{
+				const FString& InputSwitch = Switches[SwitchIndex];
+				const bool bIsFirst = InputSwitch.StartsWith(InFirstSwitch);
+				const bool bIsSecond = InputSwitch.StartsWith(InSecondSwitch);
+				if (bIsFirst || bIsSecond)
+				{
+					 // Check we are alternating
+					if ((bIsFirst && bFoundFirst)
+					 || (bIsSecond && bFoundSecond))
+					{
+						// The error will be set in the check after this loop.
+						break;
+					}
+					FString ValStr;
+					InputSwitch.Split(IToolMode::EqualsStr, nullptr, &ValStr);
+					if (bIsFirst && !ParseValue(ValStr, FirstValue))
+					{
+						ErrorMessage = FString::Printf(TEXT("Invalid parse of argument %s value: %s."),
+							*FString(FCString::Strlen(InFirstSwitch) - 1, InFirstSwitch),
+							*ValStr);
+						return false;
+					}
+					else if (bIsSecond && !ParseValue(ValStr, SecondValue))
+					{
+						ErrorMessage = FString::Printf(TEXT("Invalid parse of argument %s value: %s."),
+							*FString(FCString::Strlen(InSecondSwitch) - 1, InSecondSwitch),
+							*ValStr);
+						return false;
+					}
+					bFoundFirst |= bIsFirst;
+					bFoundSecond |= bIsSecond;
+				}
+				if (bFoundFirst && bFoundSecond)
+				{
+					OutPairs.Emplace(MoveTemp(FirstValue), MoveTemp(SecondValue));
+					bFoundAnyPair = true;
+					bFoundFirst = false;
+					bFoundSecond = false;
+					FirstValue = TFirst();
+					SecondValue = TSecond();
+				}
+			}
+			if (bFoundFirst != bFoundSecond)
+			{
+				ErrorMessage = FString::Printf(TEXT("%s and %s parameters must be provided as pairs."),
+					*FString(FCString::Strlen(InFirstSwitch) - 1, InFirstSwitch),
+					*FString(FCString::Strlen(InSecondSwitch) - 1, InSecondSwitch));
+				return false;
+			}
+			return bFoundAnyPair;
+		}
+
+		bool ParsePairs(const TCHAR* InFirstSwitch, const TCHAR* InSecondSwitch, const TArray<FString>& Switches, TMap<FGuid, TArray<uint8>>& OutPairs, FString& ErrorMessage)
+		{
+			return BuildPatchTool::IToolMode::ParsePairs<FGuid, TArray<uint8>, TMap<FGuid, TArray<uint8>>>(InFirstSwitch, InSecondSwitch, Switches, OutPairs, ErrorMessage);
+		}
+
+		bool ParsePairs(const TCHAR* InFirstSwitch, const TCHAR* InSecondSwitch, const TArray<FString>& Switches, TArray<TPair<FString, FString>>& OutPairs, FString& ErrorMessage)
+		{
+			return BuildPatchTool::IToolMode::ParsePairs<FString, FString, TArray<TPair<FString, FString>>>(InFirstSwitch, InSecondSwitch, Switches, OutPairs, ErrorMessage);
+		}
+
 		bool ParseValue(const FString& ValueIn, FString& ValueOut)
 		{
 			ValueOut = ValueIn.TrimQuotes();
@@ -140,16 +301,35 @@ namespace BuildPatchTool
 			return false;
 		}
 
-		template<typename ContainerType>
-		bool ParseValue(const FString& ValueIn, ContainerType& ValueOut)
+		bool ParseValue(const FString& ValueIn, int32& ValueOut)
+		{
+			if (FCString::IsNumeric(*ValueIn) && !ValueIn.Contains(TEXT("-"), ESearchCase::CaseSensitive))
+			{
+				ValueOut = FCString::Strtoi(*ValueIn, nullptr, 10);
+				return true;
+			}
+			return false;
+		}
+
+		bool ParseValue(const FString& ValueIn, FGuid& ValueOut)
+		{
+			return FGuid::Parse(ValueIn, ValueOut);
+		}
+
+		template<typename TCollection>
+		typename std::enable_if<is_ue_set<TCollection>::value || is_ue_array<TCollection>::value, bool>::type
+			ParseValue(const FString& ValueIn, TCollection& ValueOut)
 		{
-			TArray<FString> TempOut;
+			TArray<FString> TempArray;
 			ValueOut.Empty();
 			FString CleanValueIn = ValueIn.TrimQuotes();
-			CleanValueIn.RemoveSpacesInline();
-			if (CleanValueIn.ParseIntoArray(TempOut, TEXT(","), false) != 0)
+			if (CleanValueIn.ParseIntoArray(TempArray, TEXT(","), false) != 0)
 			{
-				ValueOut.Append(MoveTemp(TempOut));
+				for (FString& TempElem : TempArray)
+				{
+					TempElem.TrimStartAndEndInline();
+					ValueOut.Add(MoveTemp(TempElem));
+				}
 			}
 			else
 			{
@@ -157,17 +337,40 @@ namespace BuildPatchTool
 				// line and contains no commas.
 				ValueOut.Add(TEXT(""));
 			}
-
 			return true;
 		}
-	};
 
-	typedef TSharedRef<IToolMode> IToolModeRef;
-	typedef TSharedPtr<IToolMode> IToolModePtr;
+		bool ParseValue(const FString& ValueIn, TArray<uint8>& ValueOut)
+		{
+			// Here we are expecting a string into bytes, and for now, we are supporting Hex only.
+			const int32 NumChars = ValueIn.Len();
+			// An odd number of characters means the first byte is a padded nibble
+			int32 NumBytes = ValueIn.Len() / 2;
+			if (NumChars % 2 == 1)
+			{
+				++NumBytes;
+			}
+			ValueOut.Empty(NumBytes);
+			ValueOut.AddZeroed(NumBytes);
+			return ValueOut.Num() == 0 || HexToBytes(ValueIn, &ValueOut[0]) == NumBytes;
+		}
+
+		template<typename TEnum>
+		typename std::enable_if<std::is_enum<TEnum>::value, bool>::type
+			ParseValue(const FString& ValueIn, TEnum& ValueOut)
+		{
+			LexFromString(ValueOut, *ValueIn);
+			return ValueOut != TEnum::InvalidOrMax;
+		}
+	};
 
 	class FToolModeFactory
 	{
 	public:
-		static IToolModeRef Create(IBuildPatchServicesModule& BpsInterface);
+		typedef TFunction<IToolMode* (IBuildPatchServicesModule&, const TCHAR* /*CommandLine*/)> FToolModeCreateFunc;
+		typedef TMulticastDelegate<void(TMap<FString, FToolModeCreateFunc>&)> FModeRegistryDelegate;
+
+		static IToolMode* Create(IBuildPatchServicesModule& BpsInterface, const TCHAR* CommandLine);
+		static FModeRegistryDelegate& GetModeRegistryDelegate();
 	};
 }
diff --git a/Engine/Source/Programs/BuildPatchTool/Public/ToolModes/AutomationMode.h b/Engine/Source/Programs/BuildPatchTool/Public/ToolModes/AutomationMode.h
--- a/Engine/Source/Programs/BuildPatchTool/Public/ToolModes/AutomationMode.h
+++ b/Engine/Source/Programs/BuildPatchTool/Public/ToolModes/AutomationMode.h
@@ -3,16 +3,3 @@
 
 #include "CoreMinimal.h"
 #include "Interfaces/ToolMode.h"
-
-#if !UE_BUILD_SHIPPING
-
-namespace BuildPatchTool
-{
-	class FAutomationToolModeFactory
-	{
-	public:
-		static IToolModeRef Create(IBuildPatchServicesModule& BpsInterface);
-	};
-}
-
-#endif // !UE_BUILD_SHIPPING
diff --git a/Engine/Source/Programs/BuildPatchTool/Public/ToolModes/ChunkDeltaOptimiseMode.h b/Engine/Source/Programs/BuildPatchTool/Public/ToolModes/ChunkDeltaOptimiseMode.h
--- a/Engine/Source/Programs/BuildPatchTool/Public/ToolModes/ChunkDeltaOptimiseMode.h
+++ b/Engine/Source/Programs/BuildPatchTool/Public/ToolModes/ChunkDeltaOptimiseMode.h
@@ -3,12 +3,3 @@
 
 #include "CoreMinimal.h"
 #include "Interfaces/ToolMode.h"
-
-namespace BuildPatchTool
-{
-	class FChunkDeltaOptimiseToolModeFactory
-	{
-	public:
-		static IToolModeRef Create(IBuildPatchServicesModule& BpsInterface);
-	};
-}
diff --git a/Engine/Source/Programs/BuildPatchTool/Public/ToolModes/CompactifyMode.h b/Engine/Source/Programs/BuildPatchTool/Public/ToolModes/CompactifyMode.h
--- a/Engine/Source/Programs/BuildPatchTool/Public/ToolModes/CompactifyMode.h
+++ b/Engine/Source/Programs/BuildPatchTool/Public/ToolModes/CompactifyMode.h
@@ -3,12 +3,3 @@
 
 #include "CoreMinimal.h"
 #include "Interfaces/ToolMode.h"
-
-namespace BuildPatchTool
-{
-	class FCompactifyToolModeFactory
-	{
-	public:
-		static IToolModeRef Create(IBuildPatchServicesModule& BpsInterface);
-	};
-}
diff --git a/Engine/Source/Programs/BuildPatchTool/Public/ToolModes/DiffManifestMode.h b/Engine/Source/Programs/BuildPatchTool/Public/ToolModes/DiffManifestMode.h
--- a/Engine/Source/Programs/BuildPatchTool/Public/ToolModes/DiffManifestMode.h
+++ b/Engine/Source/Programs/BuildPatchTool/Public/ToolModes/DiffManifestMode.h
@@ -3,12 +3,3 @@
 
 #include "CoreMinimal.h"
 #include "Interfaces/ToolMode.h"
-
-namespace BuildPatchTool
-{
-	class FDiffManifestToolModeFactory
-	{
-	public:
-		static IToolModeRef Create(IBuildPatchServicesModule& BpsInterface);
-	};
-}
diff --git a/Engine/Source/Programs/BuildPatchTool/Public/ToolModes/EnumerationMode.h b/Engine/Source/Programs/BuildPatchTool/Public/ToolModes/EnumerationMode.h
--- a/Engine/Source/Programs/BuildPatchTool/Public/ToolModes/EnumerationMode.h
+++ b/Engine/Source/Programs/BuildPatchTool/Public/ToolModes/EnumerationMode.h
@@ -3,12 +3,3 @@
 
 #include "CoreMinimal.h"
 #include "Interfaces/ToolMode.h"
-
-namespace BuildPatchTool
-{
-	class FEnumerationToolModeFactory
-	{
-	public:
-		static IToolModeRef Create(IBuildPatchServicesModule& BpsInterface);
-	};
-}
diff --git a/Engine/Source/Programs/BuildPatchTool/Public/ToolModes/MergeManifestMode.h b/Engine/Source/Programs/BuildPatchTool/Public/ToolModes/MergeManifestMode.h
--- a/Engine/Source/Programs/BuildPatchTool/Public/ToolModes/MergeManifestMode.h
+++ b/Engine/Source/Programs/BuildPatchTool/Public/ToolModes/MergeManifestMode.h
@@ -3,12 +3,3 @@
 
 #include "CoreMinimal.h"
 #include "Interfaces/ToolMode.h"
-
-namespace BuildPatchTool
-{
-	class FMergeManifestToolModeFactory
-	{
-	public:
-		static IToolModeRef Create(IBuildPatchServicesModule& BpsInterface);
-	};
-}
diff --git a/Engine/Source/Programs/BuildPatchTool/Public/ToolModes/PackageChunksMode.h b/Engine/Source/Programs/BuildPatchTool/Public/ToolModes/PackageChunksMode.h
--- a/Engine/Source/Programs/BuildPatchTool/Public/ToolModes/PackageChunksMode.h
+++ b/Engine/Source/Programs/BuildPatchTool/Public/ToolModes/PackageChunksMode.h
@@ -3,12 +3,3 @@
 
 #include "CoreMinimal.h"
 #include "Interfaces/ToolMode.h"
-
-namespace BuildPatchTool
-{
-	class FPackageChunksToolModeFactory
-	{
-	public:
-		static IToolModeRef Create(IBuildPatchServicesModule& BpsInterface);
-	};
-}
diff --git a/Engine/Source/Programs/BuildPatchTool/Public/ToolModes/VerifyChunksMode.h b/Engine/Source/Programs/BuildPatchTool/Public/ToolModes/VerifyChunksMode.h
--- a/Engine/Source/Programs/BuildPatchTool/Public/ToolModes/VerifyChunksMode.h
+++ b/Engine/Source/Programs/BuildPatchTool/Public/ToolModes/VerifyChunksMode.h
@@ -3,12 +3,3 @@
 
 #include "CoreMinimal.h"
 #include "Interfaces/ToolMode.h"
-
-namespace BuildPatchTool
-{
-	class FVerifyChunksToolModeFactory
-	{
-	public:
-		static IToolModeRef Create(IBuildPatchServicesModule& BpsInterface);
-	};
-}
diff --git a/Engine/Source/Programs/BuildPatchTool/Resources/Windows/BuildPatchTool.rc b/Engine/Source/Programs/BuildPatchTool/Resources/Windows/BuildPatchTool.rc
new file mode 100644
--- /dev/null
+++ b/Engine/Source/Programs/BuildPatchTool/Resources/Windows/BuildPatchTool.rc
@@ -0,0 +1,82 @@
+// Copyright Epic Games, Inc. All Rights Reserved.
+
+#define APSTUDIO_READONLY_SYMBOLS
+
+#include <windows.h>
+#include "../Source/Programs/BuildPatchTool/Public/BuildPatchToolVersion.h"
+#include "../Source/Runtime/Launch/Resources/Version.h"
+#include "Runtime/Launch/Resources/Windows/resource.h"
+
+#undef APSTUDIO_READONLY_SYMBOLS
+
+#define PREPROCESSOR_TO_STRING(x) PREPROCESSOR_TO_STRING_INNER(x)
+#define PREPROCESSOR_TO_STRING_INNER(x) #x
+
+#define PORTAL_VERSION_STRING \
+	VERSION_STRINGIFY(BUILDPATCHTOOL_MAJOR_VERSION) \
+	VERSION_TEXT(".") \
+	VERSION_STRINGIFY(BUILDPATCHTOOL_MINOR_VERSION) \
+	VERSION_TEXT(".") \
+	VERSION_STRINGIFY(BUILDPATCHTOOL_PATCH_VERSION) \
+	VERSION_TEXT("-") \
+	VERSION_STRINGIFY(PORTAL_CHANGELIST) \
+	VERSION_TEXT("+") \
+	PREPROCESSOR_TO_STRING(PORTAL_BRANCH_NAME)
+
+#ifdef PROJECT_PRODUCT_NAME
+#define BUILD_PROJECT_PRODUCT_NAME PREPROCESSOR_TO_STRING(PROJECT_PRODUCT_NAME)
+#else
+#define BUILD_PROJECT_PRODUCT_NAME PREPROCESSOR_TO_STRING(UE_APP_NAME)
+#endif
+
+#ifdef PROJECT_PRODUCT_IDENTIFIER
+#define BUILD_PROJECT_PRODUCT_IDENTIFIER PREPROCESSOR_TO_STRING(PROJECT_PRODUCT_IDENTIFIER)
+#else
+#define BUILD_PROJECT_PRODUCT_IDENTIFIER EPIC_PRODUCT_IDENTIFIER
+#endif
+
+/////////////////////////////////////////////////////////////////////////////
+//
+// Version
+//
+
+VS_VERSION_INFO VERSIONINFO
+ FILEVERSION BUILDPATCHTOOL_MAJOR_VERSION, BUILDPATCHTOOL_MINOR_VERSION, BUILDPATCHTOOL_PATCH_VERSION, 0
+ PRODUCTVERSION BUILDPATCHTOOL_MAJOR_VERSION, BUILDPATCHTOOL_MINOR_VERSION, BUILDPATCHTOOL_PATCH_VERSION, 0
+ FILEFLAGSMASK 0x3fL
+#ifdef _DEBUG
+ FILEFLAGS 0x1L
+#else
+ FILEFLAGS 0x0L
+#endif
+ FILEOS 0x40004L
+ FILETYPE 0x0L
+ FILESUBTYPE 0x0L
+BEGIN
+    BLOCK "StringFileInfo"
+    BEGIN
+        BLOCK "040904b0"
+        BEGIN
+			VALUE "CompanyName", EPIC_COMPANY_NAME
+			VALUE "LegalCopyright", EPIC_COPYRIGHT_STRING
+			VALUE "ProductName", "Build Patch Tool"
+			VALUE "ProductVersion", PORTAL_VERSION_STRING
+			VALUE "FileDescription", BUILD_PROJECT_PRODUCT_NAME
+			VALUE "InternalName", BUILD_PROJECT_PRODUCT_IDENTIFIER
+			VALUE "OriginalFilename", PREPROCESSOR_TO_STRING(ORIGINAL_FILE_NAME)
+        END
+    END
+    BLOCK "VarFileInfo"
+    BEGIN
+        VALUE "Translation", 0x409, 1200
+    END
+END
+
+/////////////////////////////////////////////////////////////////////////////
+//
+// Icon
+//
+
+// Icon with lowest ID value placed first to ensure application icon
+// remains consistent on all systems.
+IDICON_UE4Game          ICON                    "UnrealEngineLauncher.ico"
diff --git a/Engine/Source/Runtime/Online/BuildPatchServices/BuildPatchServices.Build.cs b/Engine/Source/Runtime/Online/BuildPatchServices/BuildPatchServices.Build.cs
--- a/Engine/Source/Runtime/Online/BuildPatchServices/BuildPatchServices.Build.cs
+++ b/Engine/Source/Runtime/Online/BuildPatchServices/BuildPatchServices.Build.cs
@@ -5,6 +5,9 @@ using UnrealBuildTool;
 
 public class BuildPatchServices : ModuleRules
 {
+	[ConfigFile(ConfigHierarchyType.Engine, "BuildPatchServices")]
+	bool bEnableDiskOverflowStore = true;
+
 	public BuildPatchServices(ReadOnlyTargetRules Target) : base(Target)
 	{
 		StaticAnalyzerDisabledCheckers.Add("core.uninitialized.ArraySubscript");
@@ -22,6 +25,14 @@ public class BuildPatchServices : ModuleRules
 				"CoreUObject"
 			}
 		);
+		if (Target.Configuration != UnrealTargetConfiguration.Shipping)
+		{
+			PrivateDependencyModuleNames.AddRange(
+				new string[] {
+					"GoogleTest"
+				}
+			);
+		}
 
 		PrivateDependencyModuleNames.AddRange(
 		new string[] {
@@ -31,5 +42,64 @@ public class BuildPatchServices : ModuleRules
 				"Json",
 			}
 		);
+
+		PrivateIncludePathModuleNames.AddRange(
+			new string[] {
+				"HTTP"
+			}
+		);
+		PrivateIncludePathModuleNames.AddRange(
+			new string[] {
+				"HTTP"
+			}
+		);
+
+		// OpenSSL is needed on desktop platforms to support calculating SHA 256
+		bool bWithOpenSsl = false;
+		if (Target.Platform == UnrealTargetPlatform.Win64
+		 || Target.Platform == UnrealTargetPlatform.Mac
+		 || Target.Platform == UnrealTargetPlatform.Linux)
+		{
+			AddEngineThirdPartyPrivateStaticDependencies(Target, "OpenSSL");
+			bWithOpenSsl = true;
+		}
+
+		// Add a definition to allow BPS source to switch openssl code in and out
+		if (bWithOpenSsl)
+		{
+			PublicDefinitions.Add("WITH_OPENSSL=1");
+		}
+		else
+		{
+			PublicDefinitions.Add("WITH_OPENSSL=0");
+		}
+
+		if (EnableDiskOverflowStore)
+		{
+			PublicDefinitions.Add("ENABLE_PATCH_DISK_OVERFLOW_STORE=1");
+		}
+		else
+		{
+			PublicDefinitions.Add("ENABLE_PATCH_DISK_OVERFLOW_STORE=0");
+		}
+
+		if (EnableDiskOverflowStore)
+		{
+			PublicDefinitions.Add("ENABLE_PATCH_DISK_OVERFLOW_STORE=1");
+		}
+		else
+		{
+			PublicDefinitions.Add("ENABLE_PATCH_DISK_OVERFLOW_STORE=0");
+		}
+
+	}
+
+	protected bool EnableDiskOverflowStore
+	{
+		get
+		{
+			ConfigCache.ReadSettings(DirectoryReference.FromFile(Target.ProjectFile), Target.Platform, this);
+			return bEnableDiskOverflowStore;
+		}
 	}
 }
diff --git a/Engine/Source/Runtime/Online/BuildPatchServices/Private/BuildPatchDelta.cpp b/Engine/Source/Runtime/Online/BuildPatchServices/Private/BuildPatchDelta.cpp
--- a/Engine/Source/Runtime/Online/BuildPatchServices/Private/BuildPatchDelta.cpp
+++ b/Engine/Source/Runtime/Online/BuildPatchServices/Private/BuildPatchDelta.cpp
@@ -11,5 +11,5 @@ BUILDPATCHSERVICES_API FString BuildPatchServices::GetChunkDeltaFilename(const I
 
 BUILDPATCHSERVICES_API IBuildManifestPtr BuildPatchServices::MergeDeltaManifest(const IBuildManifestRef& Manifest, const IBuildManifestRef& Delta)
 {
-	return FBuildMergeManifests::MergeDeltaManifest(StaticCastSharedRef<FBuildPatchAppManifest>(Manifest), StaticCastSharedRef<FBuildPatchAppManifest>(Delta));
+	return FBuildMergeManifests::MergeDeltaManifest(StaticCastSharedRef<FBuildPatchAppManifest>(Manifest).Get(), StaticCastSharedRef<FBuildPatchAppManifest>(Delta).Get());
 }
