[Backout] - CL50979289
#rnx
Original CL Desc
-----------------------------------------------------------------
p4 integrate //UE5/Dev-Distro-5.5/Engine/Source/Programs/BuildPatchTool/... //UE5/Main/Engine/Source/Programs/BuildPatchTool/...
p4 integrate //UE5/Dev-Distro-5.5/Engine/Programs/BuildPatchTool/... //UE5/Main/Engine/Programs/BuildPatchTool/...
p4 integrate //UE5/Dev-Distro-5.5/Engine/Source/Runtime/Online/BuildPatchServices/... //UE5/Main/Engine/Source/Runtime/Online/BuildPatchServices/...

diff --git a/Engine/Programs/BuildPatchTool/Config/DefaultEngine.ini b/Engine/Programs/BuildPatchTool/Config/DefaultEngine.ini
--- a/Engine/Programs/BuildPatchTool/Config/DefaultEngine.ini
+++ b/Engine/Programs/BuildPatchTool/Config/DefaultEngine.ini
@@ -10,114 +10,3 @@ LogTextLocalizationManager=warning
 LogObj=warning
 LogExit=warning
 LogPlatformFile=warning
-LogICUInternationalization=warning
-LogHttp=error
-LogOnline=error
-
-[BuildPatchTool.ConsoleLog]
-+ExcludedCategories=LogOnline
-+ExcludedCategories=LogHttp
-
-[OnlineSubsystem]
-DefaultPlatformService=Mcp
-
-[OnlineSubsystemMcp]
-bEnabled=true
-
-[HTTP]
-HttpMaxConnectionsPerServer=0
-
-; OnlinePartyMcp is not needed, but spews warnings if you do not provide default config
-[OnlineSubsystemMcp.OnlinePartyMcp]
-PartyPendingCreateResponseTimeout=0
-PartyPendingJoinResponseTimeout=20
-PartyPendingLeaveResponseTimeout=10
-PartyPendingKickResponseTimeout=10
-PartyPendingJoinRequestTimeout=20
-
-;========== ACCOUNT SERVICE ==========
-[OnlineSubsystemMcp.OnlineIdentityMcp]
-ServiceName=account
-Protocol=https
-
-[OnlineSubsystemMcp.OnlineIdentityMcp Prod]
-Domain=account-public-service-prod03.ol.epicgames.com
-RedirectUrl=accounts.epicgames.com
-
-[OnlineSubsystemMcp.OnlineIdentityMcp Stage]
-Domain=account-public-service-stage.ol.epicgames.com
-RedirectUrl=accounts.epicgames-stage.ol.epicgames.com
-
-[OnlineSubsystemMcp.OnlineIdentityMcp GameDev]
-Domain=account-public-service-common-gamedev.ol.epicgames.net
-RedirectUrl=accounts.epicgames-gamedev.ol.epicgames.net
-
-[OnlineSubsystemMcp.OnlineIdentityMcp LocalHost]
-Domain=account-public-service-common-gamedev.ol.epicgames.net
-RedirectUrl=accounts.epicgames-gamedev.ol.epicgames.net
-
-[OnlineSubsystemMcp.OnlineIdentityMcp CI]
-Domain=account-public-service-common-ci.ol.epicgames.net
-RedirectUrl=accounts.epicgames-ci.ol.epicgames.net
-
-
-;========== ARTIFACT SERVICE ==========
-[OnlineSubsystemMcp.OnlineArtifactServiceMcp]
-ServiceName=artifact-service
-Protocol=https
-
-[OnlineSubsystemMcp.OnlineArtifactServiceMcp Prod]
-Domain=artifact-public-service-prod.ol.epicgames.com
-
-[OnlineSubsystemMcp.OnlineArtifactServiceMcp Stage]
-Domain=artifact-public-service-stage.ol.epicgames.com
-
-[OnlineSubsystemMcp.OnlineArtifactServiceMcp GameDev]
-Domain=artifact-public-service-gamedev.ol.epicgames.net
-
-[OnlineSubsystemMcp.OnlineArtifactServiceMcp CI]
-Domain=artifact-public-service-ci.ol.epicgames.net
-
-[OnlineSubsystemMcp.OnlineArtifactServiceMcp LocalHost]
-Domain=artifact-public-service-gamedev.ol.epicgames.net
-
-;========== ARTIFACT DELIVERY SERVICE ==========
-[OnlineSubsystemMcp.OnlineArtifactDeliveryServiceMcp]
-ServiceName=artifact-delivery
-Protocol=https
-
-[OnlineSubsystemMcp.OnlineArtifactDeliveryServiceMcp GameDev]
-Domain=artifact-delivery-public-service-gamedev.aeee.dev.use1a.on.epicgames.com
-
-[OnlineSubsystemMcp.OnlineArtifactDeliveryServiceMcp Prod]
-Domain=artifact-delivery-service-public-prod.ol.epicgames.com
-+AltDomains=artifact-delivery-public-service-prod.beee.live.use1a.on.epicgames.com
-
-;========== LAUNCHER SERVICE ==========
-[OnlineSubsystemMcp.OnlineBuildInfoMcpV2]
-ServiceName=launcher
-Protocol=https
-
-[OnlineSubsystemMcp.OnlineBuildInfoMcpV2 Prod]
-Domain=launcher-public-service-prod06.ol.epicgames.com
-
-[OnlineSubsystemMcp.OnlineBuildInfoMcpV2 Stage]
-Domain=launcher-public-service-stage.ol.epicgames.
-
-[OnlineSubsystemMcp.OnlineBuildInfoMcpV2 GameDev]
-Domain=launcher-public-service-gamedev.ol.epicgames.net
-
-[OnlineSubsystemMcp.OnlineBuildInfoMcpV2 CI]
-Domain=launcher-public-service-ci.ol.epicgames.net
-
-[OnlineSubsystemMcp.OnlineBuildInfoMcpV2 LocalHost]
-Domain=launcher-public-service-gamedev.ol.epicgames.net
-
-;========== CONTENT PROTECTION SERVICE ==========
-
-[OnlineSubsystemMcp.ContentProtectionServiceMcp]
-ServiceName=contentprotection
-Protocol=https
-
-[OnlineSubsystemMcp.ContentProtectionServiceMcp GameDev]
-Domain=content-protection-service-gamedev.dev.use1a.on.epicgames.net
diff --git a/Engine/Source/Programs/BuildPatchTool/BuildPatchTool.Build.cs b/Engine/Source/Programs/BuildPatchTool/BuildPatchTool.Build.cs
--- a/Engine/Source/Programs/BuildPatchTool/BuildPatchTool.Build.cs
+++ b/Engine/Source/Programs/BuildPatchTool/BuildPatchTool.Build.cs
@@ -1,13 +1,12 @@
 // Copyright Epic Games, Inc. All Rights Reserved.
 
 using UnrealBuildTool;
-using System.IO;
 
 public class BuildPatchTool : ModuleRules
 {
-	public BuildPatchTool(ReadOnlyTargetRules Target) : base(Target)
+	public BuildPatchTool( ReadOnlyTargetRules Target ) : base(Target)
 	{
-		PublicIncludePaths.Add("Runtime/Launch/Public");
+		PublicIncludePathModuleNames.Add("Launch");
 
 		PrivateDependencyModuleNames.AddRange(
 			new string[]
@@ -16,14 +15,26 @@ public class BuildPatchTool : ModuleRules
 				"CoreUObject",
 				"BuildPatchServices",
 				"Projects",
-				"Networking",
 				"HTTP",
-				"PlatformCryptoOpenSSL",
+				// The below items are not strictly needed by BPT, but core appears to need them during initialization
+				"PakFile",
+				"SandboxFile",
+				"NetworkFile",
+				"StreamingFile"
 			}
 		);
-		
-		PublicDefinitions.Add("WITH_GOOGLE_MOCK=0");
-		PublicDefinitions.Add("WITH_GOOGLE_TEST=0");
-		PublicDefinitions.Add("WITH_ONLINEMCP=0");
+
+
+		if (Target.Configuration != UnrealTargetConfiguration.Shipping)
+		{
+			PrivateDependencyModuleNames.AddRange(
+				new string[] {
+					"ApplicationCore", // Required by AutomationController
+					"Messaging",
+					"AutomationWorker",
+					"AutomationController"
+				}
+			);
+		}
 	}
 }
diff --git a/Engine/Source/Programs/BuildPatchTool/BuildPatchTool.Target.cs b/Engine/Source/Programs/BuildPatchTool/BuildPatchTool.Target.cs
--- a/Engine/Source/Programs/BuildPatchTool/BuildPatchTool.Target.cs
+++ b/Engine/Source/Programs/BuildPatchTool/BuildPatchTool.Target.cs
@@ -10,36 +10,24 @@ public class BuildPatchToolTarget : TargetRules
 	public BuildPatchToolTarget(TargetInfo Target) : base(Target)
 	{
 		Type = TargetType.Program;
-		IncludeOrderVersion = EngineIncludeOrderVersion.Latest;
-		DefaultBuildSettings = BuildSettingsVersion.Latest;
 		LinkType = TargetLinkType.Monolithic;
-
+		IncludeOrderVersion = EngineIncludeOrderVersion.Latest;
 		LaunchModuleName = "BuildPatchTool";
-		bLegalToDistributeBinary = true;
+        bLegalToDistributeBinary = true;
 		UndecoratedConfiguration = UnrealTargetConfiguration.Shipping;
 
 		bCompileAgainstEngine = false;
 		bCompileAgainstCoreUObject = false;
-		bCompileAgainstApplicationCore = false;
+		if(Target.Configuration == UnrealTargetConfiguration.Shipping)
+		{
+			// AutomationController is referenced by BuildPatchTool.build.cs, and references a ton of editor things that it probably shouldn't. Need ApplicationCore for this.
+			bCompileAgainstApplicationCore = false;
+		}
 		bBuildDeveloperTools = false;
 		bUseLoggingInShipping = true;
 		bUseChecksInShipping = true;
 		bIsBuildingConsoleApplication = true;
 		bHasExports = false;
 		bWithServerCode = false;
-		bForceDisableAutomationTests = true;
-
-		GlobalDefinitions.Add("DISABLE_CWD_CHANGES=1");
-		GlobalDefinitions.Add("UE_DONT_USE_DEFAULT_OUTPUT_DEVICES=1");
-		
-		// We don't use Crash Reporter right now.
-		GlobalDefinitions.Add("NOINITCRASHREPORTER=1");
-
-		// This is needed to avoid adding a dependency on thirdparty dynamic libraries that BPT doesn't use
-		GlobalDefinitions.Add("DISABLE_AUDIO_MCP_MODULES=1");
-
-		// Add definitions required by our windows RC file.
-		GlobalDefinitions.Add(string.Format("PORTAL_CHANGELIST={0}", Target.Version.Changelist));
-		GlobalDefinitions.Add(string.Format("PORTAL_BRANCH_NAME=\"{0}\"", Target.Version.BranchName));
 	}
 }
diff --git a/Engine/Source/Programs/BuildPatchTool/Private/BuildPatchTool.cpp b/Engine/Source/Programs/BuildPatchTool/Private/BuildPatchTool.cpp
--- a/Engine/Source/Programs/BuildPatchTool/Private/BuildPatchTool.cpp
+++ b/Engine/Source/Programs/BuildPatchTool/Private/BuildPatchTool.cpp
@@ -1,226 +1,76 @@
 // Copyright Epic Games, Inc. All Rights Reserved.
 
 #include "BuildPatchTool.h"
-
+#include "UObject/Object.h"
 #include "RequiredProgramMainCPPInclude.h"
-
-#include "Misc/App.h"
-#include "Misc/OutputDeviceError.h"
-#include "Http.h"
-#include "UObject/UObjectBase.h"
-
-#include "BuildPatchToolVersion.h"
-#include "Common/CmdLineProcessor.h"
-#include "Common/VerboseLogging.h"
-#include "Interfaces/IBuildPatchServicesModule.h"
 #include "Interfaces/ToolMode.h"
-#include "ToolModes/AutomationMode.h"
-#include "ToolModes/ChunkDeltaOptimiseMode.h"
-#include "ToolModes/ChunkDirectoryMode.h"
-#include "ToolModes/CompactifyMode.h"
-#include "ToolModes/DiffManifestMode.h"
-#include "ToolModes/ExtractMetadataMode.h"
-#include "ToolModes/EnumerationMode.h"
-#include "ToolModes/InstallManifestMode.h"
-#include "ToolModes/MergeManifestMode.h"
-#include "ToolModes/PackageChunksMode.h"
-#include "ToolModes/ToolModesHelp.h"
-#include "ToolModes/VerifyChunksMode.h"
-
-#if WITH_GOOGLE_MOCK
-THIRD_PARTY_INCLUDES_START
-#include "gmock/gmock.h"
-THIRD_PARTY_INCLUDES_END
-#endif
+#include "Misc/OutputDeviceError.h"
 
 using namespace BuildPatchTool;
 
 IMPLEMENT_APPLICATION(BuildPatchTool, "BuildPatchTool");
 DEFINE_LOG_CATEGORY(LogBuildPatchTool);
 
-namespace
+const TCHAR* HandleLegacyCommandline(const TCHAR* CommandLine)
 {
-	void SetCurrentWorkingDirectoryToBaseDir()
-	{
-		FPlatformMisc::CacheLaunchDir();
-#if PLATFORM_WINDOWS
-		TCHAR SystemError[1024];
-		verifyf(::SetCurrentDirectoryW(FWindowsPlatformProcess::BaseDir()), TEXT("Failed to set the working directory to '%s' (%s)"),
-			FWindowsPlatformProcess::BaseDir(),
-			FWindowsPlatformMisc::GetSystemErrorMessage(SystemError, UE_ARRAY_COUNT(SystemError), 0));
-#elif PLATFORM_MAC
-		chdir([FString(FMacPlatformProcess::BaseDir()).GetNSString() fileSystemRepresentation]);
-#elif PLATFORM_UNIX
-		chdir(TCHAR_TO_ANSI(FPlatformProcess::BaseDir()));
-#else
-#error Please implement SetCurrentWorkingDirectoryToBaseDir() for current platform
+	static FString CommandLineString;
+	CommandLineString = CommandLine;
+
+#if UE_BUILD_DEBUG
+	// Run smoke tests in debug
+	CommandLineString += TEXT(" -bForceSmokeTests ");
 #endif
-	}
-}
 
-namespace CommandLineHelpers
-{
-	bool ParseSwitch(const TCHAR* InSwitch, FString& Value, const TArray<FString>& Switches)
+	// For BPT, we will interpret -stdout as all log lines go to stdout.
+	if (FParse::Param(CommandLine, TEXT("stdout")))
 	{
-		// Debug check requirements for InSwitch
-		checkSlow(InSwitch != nullptr);
-		checkSlow(InSwitch[FCString::Strlen(InSwitch) - 1] == TEXT('='));
-
-		for (const FString& Switch : Switches)
-		{
-			if (Switch.StartsWith(InSwitch))
-			{
-				FString StringValue;
-				Switch.Split(TEXT("="), nullptr, &StringValue);
-				Value = StringValue.TrimQuotes();
-				return true;
-			}
-		}
-		return false;
+		CommandLineString += TEXT(" -AllowStdOutLogVerbosity ");
 	}
-}
-
-namespace WorkingDirectory
-{
-	// By default in UE4 projects the current working dir is the *.exe file location dir
-	// While using BPT it is more convenient to work with files from launch directory
-	// That's why BPT is built with DISABLE_CWD_CHANGES=1 define and the below code is added (EGS-23907)
-
-	static FString LaunchDirectoryPath;
 
-	void Configure()
+	// No longer supported options
+	if (CommandLineString.Contains(TEXT("-nochunks")))
 	{
-		LaunchDirectoryPath = FPlatformProcess::GetCurrentWorkingDirectory() + TEXT("/");
-		SetCurrentWorkingDirectoryToBaseDir();
+		UE_LOG(LogBuildPatchTool, Error, TEXT("NoChunks is no longer a supported mode. Remove this commandline option."));
+		return nullptr;
 	}
-}
 
-class FBuildPatchOutputDevice : public FOutputDevice
-{
-public:
-	FBuildPatchOutputDevice()
-		: bIsInitialized(false)
-	{}
-
-	void Initialize()
+	// Check for legacy tool mode switching, if we don't have a mode and this was not a -help request, add the correct mode
+	if (!CommandLineString.Contains(TEXT("-mode=")) && !CommandLineString.Contains(TEXT("-help")))
 	{
-		GConfig->GetArray(TEXT("BuildPatchTool.ConsoleLog"), TEXT("ExcludedCategories"), ExcludedCategories, GEngineIni);
-		bIsInitialized = true;
-	}
-
-	virtual void Serialize(const TCHAR* V, ELogVerbosity::Type Verbosity, const FName& Category) override
-	{
-		FString LogLine;
-		bool bShouldLog = false;
-		if (Verbosity <= ELogVerbosity::Display)
+		if (CommandLineString.Contains(TEXT("-compactify")))
 		{
-			FString CategoryStr = Category.ToString();
-			// We should always log our own loggers display messages.
-			bShouldLog = CategoryStr.Equals(TEXT("LogBuildPatchTool"));
-			// Once we've been initialized, we should log things that we haven't explicitly excluded.
-			bShouldLog |= bIsInitialized && !ExcludedCategories.Contains(CategoryStr);
-			if (bShouldLog)
-			{
-				// We've disabled the console logger to give us finer grained control over our log output.
-				// This message is something we want to print, so we should display without the log "clutter".
-				LogLine = V;
-			}
+			CommandLineString = CommandLineString.Replace(TEXT("-compactify"), TEXT("-mode=compactify"));
 		}
-
-		if (bShouldLog)
+		else if (CommandLineString.Contains(TEXT("-dataenumerate")))
+		{
+			CommandLineString = CommandLineString.Replace(TEXT("-dataenumerate"), TEXT("-mode=enumeration"));
+		}
+		// Patch generation did not have a mode flag, but does have some unique and required params
+		else if (CommandLineString.Contains(TEXT("-BuildRoot=")) && CommandLineString.Contains(TEXT("-BuildVersion=")))
 		{
-			FGenericPlatformMisc::LocalPrint(*(LogLine + TEXT("\n")));
-			fflush(stdout);
+			FString NewCommandline(TEXT("-mode=patchgeneration "), CommandLineString.Len());
+			NewCommandline += CommandLineString;
+			CommandLineString = MoveTemp(NewCommandline);
 		}
 	}
 
-private:
-	bool bIsInitialized;
-	TArray<FString> ExcludedCategories;
-};
-
-const FString GetRedactedCommandLine(const TCHAR* CommandLine)
-{
-    FString CommandLineString;
-    CommandLineString = CommandLine;
-    const FString Redaction = TEXT("******");
-    const TArray<FString> ParamsToRedact = { TEXT("ClientSecret") ,TEXT("EncryptionSecretKey") };
-    TArray<FString> Tokens, Switches;
-	FCommandLine::Parse(CommandLine, Tokens, Switches);
-    for (const FString& ParamToRedact : ParamsToRedact)
-    {
-        FString SearchString = ParamToRedact + TEXT("=");
-        FString Value;
-        if (CommandLineHelpers::ParseSwitch(*SearchString, Value, Switches))
-        {
-            CommandLineString.ReplaceInline(*(SearchString + Value), *(SearchString + Redaction));
-        }
-    }
-
-    return CommandLineString;
+	return *CommandLineString;
 }
 
-struct FModuleLoadHelper
-{
-public:
-	FModuleLoadHelper(TCHAR const * const InModuleName)
-		: ModuleName(InModuleName)
-	{
-		FModuleManager::Get().LoadModuleChecked(ModuleName);
-	}
-
-	~FModuleLoadHelper()
-	{
-		FModuleManager::Get().UnloadModule(ModuleName);
-	}
-
-	template <typename ModuleInterface>
-	ModuleInterface& Get() const
-	{
-		return FModuleManager::GetModuleChecked<ModuleInterface>(ModuleName);
-	}
-
-	template <typename ModuleInterface>
-	TSharedRef<ModuleInterface>& GetShared() const
-	{
-		return FModuleManager::GetModuleChecked<ModuleInterface>(ModuleName);
-	}
-
-private:
-	TCHAR const * const ModuleName;
-};
-
-EReturnCode RunBuildPatchTool(const TCHAR* CommandLine)
+EReturnCode RunBuildPatchTool()
 {
-	using namespace BuildPatchTool;
-#if !UE_BUILD_SHIPPING
-	if (FParse::Param(CommandLine, TEXT("debugging")))
-	{
-		while (!FPlatformMisc::IsDebuggerPresent())
-		{
-			FPlatformProcess::Sleep(0.01f); // to avoid 100% CPU load
-		}
-	}
-#endif
-
 	// Initialise the UObject module.
-	FModuleLoadHelper LoadedCoreUObjectModule(TEXT("CoreUObject"));
+	FModuleManager::Get().LoadModule(TEXT("CoreUObject"));
 	FCoreDelegates::OnInit.Broadcast();
 
 	// Load the BuildPatchServices Module.
-	FModuleLoadHelper LoadedBuildPatchServicesModule(TEXT("BuildPatchServices"));
-	IBuildPatchServicesModule& BuildPatchServicesModule = LoadedBuildPatchServicesModule.Get<IBuildPatchServicesModule>();
+	IBuildPatchServicesModule& BuildPatchServicesModule = FModuleManager::LoadModuleChecked<IBuildPatchServicesModule>(TEXT("BuildPatchServices"));
 
 	// Make sure we have processed UObjects from BPS.
 	ProcessNewlyLoadedUObjects();
 
-	// Unfortunately we cannot control UE logging going to GWarn as well as our own simpler log.
-	class FFeedbackContext* GWarnBackup = GWarn;
-	GWarn = nullptr;
-	ON_SCOPE_EXIT{ GWarn = GWarnBackup; };
-
-	// Run the tool mode.
-	TUniquePtr<IToolMode> ToolMode(FToolModeFactory::Create(BuildPatchServicesModule, CommandLine));
+	// Instantiate and execute the tool.
+	TSharedRef<IToolMode> ToolMode = FToolModeFactory::Create(BuildPatchServicesModule);
 	return ToolMode->Execute();
 }
 
@@ -228,12 +78,8 @@ int32 NumberOfWorkerThreadsDesired()
 {
 	const int32 MaxThreads = 64;
 	const int32 NumberOfCores = FPlatformMisc::NumberOfCores();
-	int32 NumThreadsOverride = NumberOfCores;
-	FParse::Value(FCommandLine::Get(), TEXT("NumWorkerThreads="), NumThreadsOverride);
-	int32 DesiredNumberOfThreads = FMath::Max(NumberOfCores - 1, NumThreadsOverride);
-	UE_LOG(LogBuildPatchTool, Log, TEXT("Desired number of threads is [%d] based on NumberOfCores [%d] and NumThreadsOverride [%d]."), DesiredNumberOfThreads, NumberOfCores, NumThreadsOverride);
 	// need to spawn at least one worker thread (see FTaskGraphImplementation)
-	return FMath::Max(FMath::Min(DesiredNumberOfThreads, MaxThreads), 1);
+	return FMath::Max(FMath::Min(NumberOfCores - 1, MaxThreads), 1);
 }
 
 void CheckAndReallocThreadPool()
@@ -242,7 +88,6 @@ void CheckAndReallocThreadPool()
 	{
 		const int32 ThreadsSpawned = GThreadPool->GetNumThreads();
 		const int32 DesiredThreadCount = NumberOfWorkerThreadsDesired();
-
 		if (ThreadsSpawned < DesiredThreadCount)
 		{
 			UE_LOG(LogBuildPatchTool, Log, TEXT("Engine only spawned %d worker threads, bumping up to %d!"), ThreadsSpawned, DesiredThreadCount);
@@ -257,95 +102,27 @@ void CheckAndReallocThreadPool()
 	}
 }
 
-EReturnCode BuildPatchToolMain(int32 ArgC, TCHAR* ArgV[])
+EReturnCode BuildPatchToolMain(const TCHAR* CommandLine)
 {
-	FTaskTagScope Scope(ETaskTag::EGameThread);
-	FString EngineExitReason = TEXT("Exiting");
-	ON_SCOPE_EXIT
-	{
-		RequestEngineExit(*EngineExitReason);
-	};
-
-	GAlwaysReportCrash = true;
-	FApp::SetProjectName(TEXT("BuildPatchTool"));
-	FCmdLineProcessor CmdLineProcessor(ArgC, ArgV, WorkingDirectory::LaunchDirectoryPath, [](const FString& LogLine)
-	{
-		FGenericPlatformMisc::LocalPrint(*(LogLine + TEXT("\n")));
-		UE_LOG(LogBuildPatchTool, Error, TEXT("%s"), *LogLine);
-	});
-	FCommandLine::Set(*CmdLineProcessor.GetEngineString());
-
-	CmdLineProcessor.ProcessDragAndDrop();
-
-	if (!CmdLineProcessor.ProcessCommandlineFiles())
-	{
-		EngineExitReason = TEXT("BuildPatchToolMain ProcessCommandlineFileArgs failed.");
-		return EReturnCode::ArgumentProcessingError;
-	}
-
-	if (!CmdLineProcessor.HandleLegacyCommandline())
-	{
-		EngineExitReason = TEXT("BuildPatchToolMain HandleLegacyCommandline failed.");
-		return EReturnCode::ArgumentProcessingError;
-	}
-
-	if (!CmdLineProcessor.ConvertCommandlinePaths())
+	// Handle legacy commandlines
+	CommandLine = HandleLegacyCommandline(CommandLine);
+	if (CommandLine == nullptr)
 	{
-		EngineExitReason = TEXT("BuildPatchToolMain ConvertCommandlinePaths failed.");
 		return EReturnCode::ArgumentProcessingError;
 	}
 
-	CmdLineProcessor.AddDesiredEngineFlags();
-
-	bool bVerboseLogsSpecified = CmdLineProcessor.ContainsFlag(TEXT("-verboseLogs"));
-	SetupVerboseLogging(bVerboseLogsSpecified);
-
-	// Add log device for stdout
-	bool bStdOutSpecified = CmdLineProcessor.ContainsFlag(TEXT("-stdout"));
-	FBuildPatchOutputDevice* BuildPatchOutputDevice = nullptr;
-	if (!bStdOutSpecified)
-	{
-		BuildPatchOutputDevice = new FBuildPatchOutputDevice();
-		GLog->AddOutputDevice(BuildPatchOutputDevice);
-	}
-
 	// Initialise application
-	const FString FullCommandLine = CmdLineProcessor.GetEngineString();
-	GEngineLoop.PreInit(*FullCommandLine);
-	ON_SCOPE_EXIT
-	{
-		FEngineLoop::AppPreExit();
-		FModuleManager::Get().UnloadModulesAtShutdown();
-		FEngineLoop::AppExit();
-	};
-	if (BuildPatchOutputDevice != nullptr)
-	{
-		BuildPatchOutputDevice->Initialize();
-	}
-
-	// Explicitly set the current culture to the supported value
-	FInternationalization::Get().SetCurrentCulture("en");
-
-	UE_LOG(LogBuildPatchTool, Log, TEXT("BuildPatchTool v%s %s"), BUILDPATCHTOOL_VERSION_STRING, LexToString(FApp::GetBuildConfiguration()));
-	UE_LOG(LogBuildPatchTool, Log, TEXT("Received  commandline: %s"), *(GetRedactedCommandLine(*CmdLineProcessor.GetOriginalCommandLine())))
-	UE_LOG(LogBuildPatchTool, Log, TEXT("Processed commandline: %s"), *(GetRedactedCommandLine(*FullCommandLine)));
-
-	// Allow QA to test crashes are handled appropriately
-	bool bForceCrashSpecified = CmdLineProcessor.ContainsFlag(TEXT("-forceCrash"));
-	if (bForceCrashSpecified)
-	{
-		*((volatile int32*)0) = 0;
-	}
+	GEngineLoop.PreInit(CommandLine);
+	UE_LOG(LogBuildPatchTool, Log, TEXT("Executed with commandline: %s"), CommandLine);
 
 	// Check whether as a program, we should bump up the number of threads in GThreadPool.
 	CheckAndReallocThreadPool();
 
 	// Run the application
-	EReturnCode ReturnCode = RunBuildPatchTool(*FullCommandLine);
+	EReturnCode ReturnCode = RunBuildPatchTool();
 	if (ReturnCode != EReturnCode::OK)
 	{
-		UE_LOG(LogBuildPatchTool, Error, TEXT("Tool exited with %s (%d)"), LexToString(ReturnCode), (int32)ReturnCode);
-		UE_LOG(LogBuildPatchTool, Error, TEXT("For more details, see log file at %s"), *FPaths::ConvertRelativePathToFull(FPlatformOutputDevices::GetAbsoluteLogFilename()));
+		UE_LOG(LogBuildPatchTool, Error, TEXT("Tool exited with: %d"), (int32)ReturnCode);
 	}
 
 	// Shutdown
@@ -358,17 +135,34 @@ EReturnCode BuildPatchToolMain(int32 ArgC, TCHAR* ArgV[])
 	return ReturnCode;
 }
 
-INT32_MAIN_INT32_ARGC_TCHAR_ARGV()
+const TCHAR* ProcessApplicationCommandline(int32 ArgC, TCHAR* ArgV[])
 {
+	static FString CommandLine = TEXT("-usehyperthreading -UNATTENDED");
+	for (int32 Option = 1; Option < ArgC; Option++)
+	{
+		CommandLine += TEXT(" ");
+		FString Argument(ArgV[Option]);
+		if (Argument.Contains(TEXT(" ")))
+		{
+			if (Argument.Contains(TEXT("=")))
+			{
+				FString ArgName;
+				FString ArgValue;
+				Argument.Split(TEXT("="), &ArgName, &ArgValue);
+				Argument = FString::Printf(TEXT("%s=\"%s\""), *ArgName, *ArgValue);
+			}
+			else
+			{
+				Argument = FString::Printf(TEXT("\"%s\""), *Argument);
+			}
+		}
+		CommandLine += Argument;
+	}
+	return *CommandLine;
+}
 
-#if WITH_GOOGLE_MOCK
-	int argc = 0;
-	wchar_t** argv = nullptr;
-	::testing::InitGoogleMock(&argc, argv);
-#endif
-	
-	WorkingDirectory::Configure();
-
+INT32_MAIN_INT32_ARGC_TCHAR_ARGV()
+{
 	EReturnCode ReturnCode;
 	// Using try&catch is the windows-specific method of interfacing with CrashReportClient
 #if PLATFORM_WINDOWS && !PLATFORM_SEH_EXCEPTIONS_DISABLED
@@ -378,7 +172,7 @@ INT32_MAIN_INT32_ARGC_TCHAR_ARGV()
 		// SetCrashHandler(nullptr) sets up default behavior for Linux and Mac interfacing with CrashReportClient
 		FPlatformMisc::SetCrashHandler(nullptr);
 		GIsGuarded = 1;
-		ReturnCode = BuildPatchToolMain(ArgC, ArgV);
+		ReturnCode = BuildPatchToolMain(ProcessApplicationCommandline(ArgC, ArgV));
 		GIsGuarded = 0;
 	}
 #if PLATFORM_WINDOWS && !PLATFORM_SEH_EXCEPTIONS_DISABLED
diff --git a/Engine/Source/Programs/BuildPatchTool/Private/ToolMode.cpp b/Engine/Source/Programs/BuildPatchTool/Private/ToolMode.cpp
--- a/Engine/Source/Programs/BuildPatchTool/Private/ToolMode.cpp
+++ b/Engine/Source/Programs/BuildPatchTool/Private/ToolMode.cpp
@@ -1,21 +1,25 @@
 // Copyright Epic Games, Inc. All Rights Reserved.
 
 #include "Interfaces/ToolMode.h"
-
-#include "BuildPatchTool.h"
-#include "BuildPatchToolVersion.h"
-#include "Common/MetadataSerialiser.h"
-#include "Interfaces/IBuildPatchServicesModule.h"
 #include "Misc/CommandLine.h"
-#include "ToolModes/ToolModesHelp.h"
+#include "Interfaces/IBuildPatchServicesModule.h"
+#include "ToolModes/PatchGenerationMode.h"
+#include "ToolModes/ChunkDeltaOptimiseMode.h"
+#include "ToolModes/CompactifyMode.h"
+#include "ToolModes/EnumerationMode.h"
+#include "ToolModes/MergeManifestMode.h"
+#include "ToolModes/DiffManifestMode.h"
+#include "ToolModes/PackageChunksMode.h"
+#include "ToolModes/VerifyChunksMode.h"
+#include "ToolModes/AutomationMode.h"
+#include "BuildPatchTool.h"
 
 namespace BuildPatchTool
 {
 	class FHelpToolMode : public IToolMode
 	{
 	public:
-		FHelpToolMode(const TCHAR* InCommandLine)
-			: CommandLine(InCommandLine)
+		FHelpToolMode()
 		{}
 
 		virtual ~FHelpToolMode()
@@ -23,52 +27,79 @@ namespace BuildPatchTool
 
 		virtual EReturnCode Execute() override
 		{
-			bool bRequestedHelp = FParse::Param(CommandLine, TEXT("help"));
+			bool bRequestedHelp = FParse::Param(FCommandLine::Get(), TEXT("help"));
 
 			// Output generic help info
 			if (!bRequestedHelp)
 			{
-				UE_LOG(LogBuildPatchTool, Error, TEXT("ERROR: No mode was detected on the commandline, please check the provided args. You can use -help for information on available modes."));
-				UE_LOG(LogBuildPatchTool, Display, TEXT(""));
+				UE_LOG(LogBuildPatchTool, Error, TEXT("No supported mode detected."));
 			}
-			PrintHelp<FBptBaseToolModeHelp>();
+			UE_LOG(LogBuildPatchTool, Display, TEXT("-help can be added with any mode selection to get extended information."));
+			UE_LOG(LogBuildPatchTool, Display, TEXT("Supported modes are:"));
+			UE_LOG(LogBuildPatchTool, Display, TEXT("  -mode=PatchGeneration      Mode that generates patch data for the a new build."));
+			UE_LOG(LogBuildPatchTool, Display, TEXT("  -mode=ChunkDeltaOptimise   Mode that produces a more optimal chunked patch between two specific builds."));
+			UE_LOG(LogBuildPatchTool, Display, TEXT("  -mode=Compactify           Mode that can clean up unneeded patch data from a given cloud directory with redundant data."));
+			UE_LOG(LogBuildPatchTool, Display, TEXT("  -mode=Enumeration          Mode that outputs the paths to referenced patch data given a single manifest."));
+			UE_LOG(LogBuildPatchTool, Display, TEXT("  -mode=MergeManifests       Mode that can combine two manifest files to create a new one, primarily used to create hotfixes."));
+			UE_LOG(LogBuildPatchTool, Display, TEXT("  -mode=DiffManifests        Mode that can diff two manifests and outputs what chunks would need to be downloaded and some stats."));
+			UE_LOG(LogBuildPatchTool, Display, TEXT("  -mode=PackageChunks        Mode that packages data required for an installation into larger files which can be used as local sources for build patch installers."));
+			UE_LOG(LogBuildPatchTool, Display, TEXT("  -mode=VerifyChunks         Mode that allows you to verify the integrity of patch data. It will load chunk or chunkdb files to verify they are not corrupt."));
+#if !UE_BUILD_SHIPPING
+			UE_LOG(LogBuildPatchTool, Display, TEXT("  -mode=AutomationTests      Mode that will run automation testing."));
+#endif // !UE_BUILD_SHIPPING
 
 			// Error if this wasn't just a help request
 			return bRequestedHelp ? EReturnCode::OK : EReturnCode::UnknownToolMode;
 		}
-
-	private:
-		const TCHAR* CommandLine;
 	};
 
-	IToolMode* FToolModeFactory::Create(IBuildPatchServicesModule& BpsInterface, const TCHAR* CommandLine)
+	IToolModeRef FToolModeFactory::Create(IBuildPatchServicesModule& BpsInterface)
 	{
-		TMap<FString, FToolModeCreateFunc> Modes;
-		GetModeRegistryDelegate().Broadcast(Modes);
-
-		// Create the correct tool mode for the commandline given.
+		// Create the correct tool mode for the commandline given
 		FString ToolModeValue;
-		if (FParse::Value(CommandLine, TEXT("mode="), ToolModeValue))
+		if (FParse::Value(FCommandLine::Get(), TEXT("mode="), ToolModeValue))
 		{
-			if (const FToolModeCreateFunc* ToolModeCreator = Modes.Find(*ToolModeValue))
+			if (ToolModeValue == TEXT("patchgeneration"))
 			{
-				return (*ToolModeCreator)(BpsInterface, CommandLine);
+				return FPatchGenerationToolModeFactory::Create(BpsInterface);
 			}
+			else if (ToolModeValue == TEXT("chunkdeltaoptimise"))
+			{
+				return FChunkDeltaOptimiseToolModeFactory::Create(BpsInterface);
+			}
+			else if (ToolModeValue == TEXT("compactify"))
+			{
+				return FCompactifyToolModeFactory::Create(BpsInterface);
+			}
+			else if (ToolModeValue == TEXT("enumeration"))
+			{
+				return FEnumerationToolModeFactory::Create(BpsInterface);
+			}
+			else if (ToolModeValue == TEXT("mergemanifests"))
+			{
+				return FMergeManifestToolModeFactory::Create(BpsInterface);
+			}
+			else if (ToolModeValue == TEXT("diffmanifests"))
+			{
+				return FDiffManifestToolModeFactory::Create(BpsInterface);
+			}
+			else if (ToolModeValue == TEXT("packagechunks"))
+			{
+				return FPackageChunksToolModeFactory::Create(BpsInterface);
+			}
+			else if (ToolModeValue == TEXT("verifychunks"))
+			{
+				return FVerifyChunksToolModeFactory::Create(BpsInterface);
+			}
+#if !UE_BUILD_SHIPPING
+			else if (ToolModeValue == TEXT("automationtests"))
+			{
+				return FAutomationToolModeFactory::Create(BpsInterface);
+			}
+#endif // !UE_BUILD_SHIPPING
 		}
-
-		// If we didn't find a toolmode, check for empty mode name for overriden help.
-		if (const FToolModeCreateFunc* ToolModeCreator = Modes.Find(TEXT("")))
-		{
-			return (*ToolModeCreator)(BpsInterface, CommandLine);
-		}
-
-		// No previous attempts found a tool mode. Create the generic help, which will return ok if -help was provided else return UnknownToolMode error
-		return new FHelpToolMode(CommandLine);
-	}
-
-	FToolModeFactory::FModeRegistryDelegate& FToolModeFactory::GetModeRegistryDelegate()
-	{
-		static FToolModeFactory::FModeRegistryDelegate ModeRegistryDelegate;
-		return ModeRegistryDelegate;
+		
+		// No mode provided so create the generic help, which will return ok if -help was provided else return UnknownToolMode error
+		return MakeShareable(new FHelpToolMode());
 	}
 }
diff --git a/Engine/Source/Programs/BuildPatchTool/Private/ToolModes/AutomationMode.cpp b/Engine/Source/Programs/BuildPatchTool/Private/ToolModes/AutomationMode.cpp
--- a/Engine/Source/Programs/BuildPatchTool/Private/ToolModes/AutomationMode.cpp
+++ b/Engine/Source/Programs/BuildPatchTool/Private/ToolModes/AutomationMode.cpp
@@ -2,74 +2,27 @@
 
 #include "ToolModes/AutomationMode.h"
 
-#if !UE_BUILD_SHIPPING && WITH_DEV_AUTOMATION_TESTS
+#if !UE_BUILD_SHIPPING
 
-#include "BuildPatchTool.h"
-#include "Containers/Ticker.h"
-#include "IAutomationWorkerModule.h"
-#include "IAutomationControllerModule.h"
 #include "Interfaces/IBuildPatchServicesModule.h"
+#include "BuildPatchTool.h"
 #include "Misc/CommandLine.h"
 #include "Misc/Paths.h"
+
+#include "IAutomationWorkerModule.h"
+#include "IAutomationControllerModule.h"
 #include "Modules/ModuleManager.h"
-#include "ToolModes/HelpToolMode.h"
-#include "ToolModes/ToolModesHelp.h"
 #include "UObject/UObjectGlobals.h"
-
-DECLARE_LOG_CATEGORY_EXTERN(LogAutomationMode, Warning, All);
-DEFINE_LOG_CATEGORY(LogAutomationMode);
+#include "Containers/Ticker.h"
 
 using namespace BuildPatchTool;
 
-#if WITH_GOOGLE_TEST
-
-THIRD_PARTY_INCLUDES_START
-#include "gtest/gtest.h"
-THIRD_PARTY_INCLUDES_END
-
-class FGoogleTestRunnerPrinter
-	: public ::testing::EmptyTestEventListener
-{
-	// Called after a failed assertion or a SUCCEED() invocation.
-	virtual void OnTestPartResult(const ::testing::TestPartResult& test_part_result) override
-	{
-		if (FAutomationTestBase* const CurrentTest = FAutomationTestFramework::GetInstance().GetCurrentTest())
-		{
-			if (test_part_result.passed())
-			{
-				CurrentTest->AddInfo(test_part_result.message());
-			}
-			else if (test_part_result.skipped())
-			{
-				CurrentTest->AddWarning(test_part_result.message());
-			}
-			else
-			{
-				CurrentTest->AddError(FString::Printf(TEXT("\n%s File:%s Line:%s"), test_part_result.message(), test_part_result.file_name(), test_part_result.line_number()));
-			}
-
-		}
-		else
-		{
-			UE_LOG(LogAutomationMode, Warning, TEXT("GoogleTest expectation or assert called when no automation test was actively running!"));
-		}
-	}
-};
-
-#endif
 class FAutomationToolMode : public IToolMode
 {
 public:
-	FAutomationToolMode(IBuildPatchServicesModule& InBpsInterface, const TCHAR* InCommandLine)
+	FAutomationToolMode(IBuildPatchServicesModule& InBpsInterface)
 		: BpsInterface(InBpsInterface)
-		, CommandLine(InCommandLine)
 	{
-#if WITH_GOOGLE_TEST
-		::testing::TestEventListeners& Listeners = ::testing::UnitTest::GetInstance()->listeners();
-		Listeners.Append(new FGoogleTestRunnerPrinter());
-		Listeners.Release(Listeners.default_result_printer());
-#endif
-
 	}
 
 	virtual ~FAutomationToolMode()
@@ -87,7 +40,16 @@ public:
 		// Print help if requested
 		if (bHelp)
 		{
-			PrintHelp<FAutomationToolModeHelp>();
+			UE_LOG(LogBuildPatchTool, Display, TEXT("AUTOMATION TEST MODE"));
+			UE_LOG(LogBuildPatchTool, Display, TEXT("This tool mode runs automation tests."));
+			UE_LOG(LogBuildPatchTool, Display, TEXT(""));
+			UE_LOG(LogBuildPatchTool, Display, TEXT("No arguments are required."));
+			UE_LOG(LogBuildPatchTool, Display, TEXT(""));
+			UE_LOG(LogBuildPatchTool, Display, TEXT("Optional arguments:"));
+			UE_LOG(LogBuildPatchTool, Display, TEXT("  -TestList=\"\"   Specifies in quotes, the list of tests to run. The list is + delimited."));
+			UE_LOG(LogBuildPatchTool, Display, TEXT(""));
+			UE_LOG(LogBuildPatchTool, Display, TEXT("NB: If -TestList is not specified, then all BuildPatchServices tests are ran."));
+			UE_LOG(LogBuildPatchTool, Display, TEXT(""));
 			return EReturnCode::OK;
 		}
 
@@ -167,7 +129,7 @@ private:
 	{
 #define PARSE_SWITCH(Switch) ParseSwitch(TEXT(#Switch "="), Switch, Switches)
 		TArray<FString> Tokens, Switches;
-		FCommandLine::Parse(CommandLine, Tokens, Switches);
+		FCommandLine::Parse(FCommandLine::Get(), Tokens, Switches);
 
 		bHelp = ParseOption(TEXT("help"), Switches);
 		if (bHelp)
@@ -176,7 +138,7 @@ private:
 		}
 		if (!PARSE_SWITCH(TestList) || TestList.Contains(TEXT(";")))
 		{
-			TestList = TEXT("BuildPatchServices+BuildPatchTool");
+			TestList = TEXT("BuildPatchServices");
 		}
 		TestList.InsertAt(0, TEXT("Automation RunTests "));
 
@@ -185,13 +147,14 @@ private:
 	}
 
 private:
-	static TCHAR const* const MODE_NAME;
 	IBuildPatchServicesModule& BpsInterface;
-	const TCHAR* CommandLine;
 	bool bHelp;
 	FString TestList;
 };
 
-IMPLEMENT_BPT_MODE(AutomationTests, FAutomationToolMode);
+BuildPatchTool::IToolModeRef BuildPatchTool::FAutomationToolModeFactory::Create(IBuildPatchServicesModule& BpsInterface)
+{
+	return MakeShareable(new FAutomationToolMode(BpsInterface));
+}
 
 #endif // !UE_BUILD_SHIPPING
diff --git a/Engine/Source/Programs/BuildPatchTool/Private/ToolModes/ChunkDeltaOptimiseMode.cpp b/Engine/Source/Programs/BuildPatchTool/Private/ToolModes/ChunkDeltaOptimiseMode.cpp
--- a/Engine/Source/Programs/BuildPatchTool/Private/ToolModes/ChunkDeltaOptimiseMode.cpp
+++ b/Engine/Source/Programs/BuildPatchTool/Private/ToolModes/ChunkDeltaOptimiseMode.cpp
@@ -1,23 +1,20 @@
 // Copyright Epic Games, Inc. All Rights Reserved.
 
 #include "ToolModes/ChunkDeltaOptimiseMode.h"
-
-#include "BuildPatchTool.h"
-#include "Common/DiffAbortThresholdArgument.h"
-#include "Interfaces/ToolMode.h"
 #include "Misc/CommandLine.h"
-#include "Misc/Optional.h"
 #include "Misc/Paths.h"
-#include "ToolModes/ToolModesHelp.h"
+#include "Interfaces/IBuildPatchServicesModule.h"
+#include "BuildPatchTool.h"
 
 using namespace BuildPatchTool;
 
 class FChunkDeltaOptimiseMode : public IToolMode
 {
 public:
-	FChunkDeltaOptimiseMode(IBuildPatchServicesModule& InBpsInterface, const TCHAR* InCommandLine)
+	FChunkDeltaOptimiseMode(IBuildPatchServicesModule& InBpsInterface)
 		: BpsInterface(InBpsInterface)
-		, CommandLine(InCommandLine)
+		, ScanWindowSize(8191)
+		, OutputChunkSize(1024 * 1024)
 	{}
 
 	virtual ~FChunkDeltaOptimiseMode()
@@ -25,9 +22,7 @@ public:
 
 	virtual EReturnCode Execute() override
 	{
-		using namespace BuildPatchServices;
-
-		// Parse command line.
+		// Parse commandline.
 		if (ProcessCommandline() == false)
 		{
 			return EReturnCode::ArgumentProcessingError;
@@ -36,96 +31,44 @@ public:
 		// Print help if requested.
 		if (bHelp)
 		{
-			PrintHelp<FChunkDeltaOptimiseModeHelp>();
+			UE_LOG(LogBuildPatchTool, Display, TEXT("CHUNK DELTA OPTIMISE MODE"));
+			UE_LOG(LogBuildPatchTool, Display, TEXT("This tool supports optimising chunk based patches to reduce the number of chunks required to download when patching between specific versions."));
+			UE_LOG(LogBuildPatchTool, Display, TEXT(""));
+			UE_LOG(LogBuildPatchTool, Display, TEXT("Required arguments:"));
+			UE_LOG(LogBuildPatchTool, Display, TEXT("  -mode=ChunkDeltaOptimise  Must be specified to launch the tool in chunk delta optimise mode."));
+			UE_LOG(LogBuildPatchTool, Display, TEXT("  -ManifestA=\"\"             Specifies in quotes the file path to the base manifest."));
+			UE_LOG(LogBuildPatchTool, Display, TEXT("  -ManifestB=\"\"             Specifies in quotes the file path to the update manifest."));
+			UE_LOG(LogBuildPatchTool, Display, TEXT(""));
+			UE_LOG(LogBuildPatchTool, Display, TEXT("Optional arguments:"));
+			UE_LOG(LogBuildPatchTool, Display, TEXT("  -CloudDir=\"\"                     Specifies in quotes the cloud directory where existing data will be recognized from, and new data added to. If not provided, location of ManifestB will be used."));
+			UE_LOG(LogBuildPatchTool, Display, TEXT("  -ScanWindowSize=1000000          Specifies in bytes, the scan window to use. Range accepted is 128KiB >= n >= 8Kb, defaults to 8191 (Closest prime to 8KiB)."));
+			UE_LOG(LogBuildPatchTool, Display, TEXT("  -OutputChunkSize=1000000         Specifies in bytes, the chunk size to save out unknown data with. Range accepted is 10MiB >= n >= 1MB, defaults to 1048576 (1MiB)."));
+			UE_LOG(LogBuildPatchTool, Display, TEXT("  -DiffAbortThreshold=10000000000  Specifies in bytes, an upper limit for original diffs to try to enhance."));
+			UE_LOG(LogBuildPatchTool, Display, TEXT("                                   This allows short circuiting lengthy optimisation attempts on large diffs which may not benefit. Range accepted is n >= 1GB, defaults to never aborting."));
+			UE_LOG(LogBuildPatchTool, Display, TEXT(""));
 			return EReturnCode::OK;
 		}
 
-		LogConfig();
-
-		// Register secrets.
-		for (const TPair<FGuid, TArray<uint8>>& IdKeyPair : EncryptionSecrets)
-		{
-			BpsInterface.RegisterEncryptionSecret(IdKeyPair.Key, IdKeyPair.Value);
-		}
-
-		// Set any commandline based config.
-		GConfig->SetInt(TEXT("Portal.BuildPatch"), TEXT("ChunkRetries"), MaxRetryCount, GEngineIni);
-
 		// Setup and run.
-		FChunkDeltaOptimiserConfiguration Configuration;
+		BuildPatchServices::FChunkDeltaOptimiserConfiguration Configuration;
 		Configuration.ManifestAUri = ManifestA;
 		Configuration.ManifestBUri = ManifestB;
-		Configuration.OutputCloudDirectory = CloudDir;
-		Configuration.DownloadCloudDirectories = DownloadCloudDirectories;
+		Configuration.CloudDirectory = CloudDir;
 		Configuration.ScanWindowSize = ScanWindowSize;
 		Configuration.OutputChunkSize = OutputChunkSize;
 		Configuration.DiffAbortThreshold = DiffAbortThreshold;
-		Configuration.DeltaFilenameTrailer = DeltaFilenameTrailer;
-		Configuration.SourceDetailsLogFilename = SourceDetailsLogFilename;
-		Configuration.ManifestBTagAliases = MoveTemp(ManifestBTagAliases);
-		Configuration.ManifestATagAliases = MoveTemp(ManifestATagAliases);
 
 		// Run the build generation.
-		IChunkDeltaOptimiserRef ChunkDeltaOptimiser = BpsInterface.CreateChunkDeltaOptimiser(Configuration);
-
-		bool bSuccess = ChunkDeltaOptimiser->Run();
-
+		bool bSuccess = BpsInterface.OptimiseChunkDelta(Configuration);
 		return bSuccess ? EReturnCode::OK : EReturnCode::ToolFailure;
 	}
 
 private:
-	// returns false on an error, not if the switch is present! Check OutMap for presence.
-	bool ParseManyToOneMap(const TCHAR* SwitchName, const TArray<FString>& Switches, TMap<FString, FString>& OutMap)
-	{
-		TArray<FString> RawStrings;
-		FString SwitchNameWithEq = FString(SwitchName) + TEXT("=");
-		if (ParseSwitches(*SwitchNameWithEq, RawStrings, Switches))
-		{
-			for (const FString& ManyToOneString : RawStrings)
-			{
-				FString From, To;
-				if (!ManyToOneString.Split(TEXT(";"), &From, &To))
-				{
-					UE_LOG(LogBuildPatchTool, Error, TEXT("%s must contain a ';' separating the 'from' tags and the 'to' tag (%s)"), SwitchName, *ManyToOneString);
-					return false;
-				}
-
-				TArray<FString> FromArray;
-				From.ParseIntoArray(FromArray, TEXT(","));
-				if (FromArray.Num() == 0)
-				{
-					UE_LOG(LogBuildPatchTool, Error, TEXT("%s must contain a comma separated list of 'from' tags (%s)"), SwitchName, *ManyToOneString);
-					return false;
-				}
-				if (To.Len() == 0)
-				{
-					UE_LOG(LogBuildPatchTool, Error, TEXT("%s must contain a 'to' tag after a ';' (%s)"), SwitchName, *ManyToOneString);
-					return false;
-				}
-
-				for (const FString& FromEntry : FromArray)
-				{
-					OutMap.Add(FromEntry, To);
-				}
-			}
-		}
-		return true;
-	}
-
 	bool ProcessCommandline()
 	{
 #define PARSE_SWITCH(Switch) ParseSwitch(TEXT(#Switch "="), Switch, Switches)
-#define _PARSE_REQUIRED_SWITCH(Switch, bExpect) \
-		if (PARSE_SWITCH(Switch) != bExpect) \
-		{ \
-			UE_LOG(LogBuildPatchTool, Error, #Switch SENTENCE_END); \
-			bParametersOk = false; \
-		}
-#define PARSE_REQUIRED_SWITCH(Switch) _PARSE_REQUIRED_SWITCH(Switch, true)
-#define PARSE_UNEXPECTED_SWITCH(Switch) _PARSE_REQUIRED_SWITCH(Switch, false)
-
 		TArray<FString> Tokens, Switches;
-		FCommandLine::Parse(CommandLine, Tokens, Switches);
+		FCommandLine::Parse(FCommandLine::Get(), Tokens, Switches);
 
 		bHelp = ParseOption(TEXT("help"), Switches);
 		if (bHelp)
@@ -133,68 +76,35 @@ private:
 			return true;
 		}
 
-		// TODOBPTONLINE
-#undef SENTENCE_END
-#define SENTENCE_END TEXT(" is a required parameter for offline mode.")
-		bool bParametersOk = true;
-		PARSE_REQUIRED_SWITCH(ManifestA);
-		PARSE_REQUIRED_SWITCH(ManifestB);
-#undef SENTENCE_END
-#define SENTENCE_END TEXT(" should only be used when uploading to Epic services.")
-		FString BuildVersionA, BuildVersionB;
-		PARSE_UNEXPECTED_SWITCH(BuildVersionA);
-		PARSE_UNEXPECTED_SWITCH(BuildVersionB);
-#undef SENTENCE_END
-#define SENTENCE_END TEXT(" is only supported when Epic services are available. Download an official build from Epic to enable this functionality.")
-		FString ClientId, OrganizationId, ProductId;
-		PARSE_UNEXPECTED_SWITCH(ClientId);
-		PARSE_UNEXPECTED_SWITCH(OrganizationId);
-		PARSE_UNEXPECTED_SWITCH(ProductId);
-
-		if (!bParametersOk)
-		{
-			return false;
-		}
-
-		if (ManifestA.IsEmpty() || ManifestB.IsEmpty())
+		// Get all required parameters.
+		if (!(PARSE_SWITCH(ManifestA) && PARSE_SWITCH(ManifestB)))
 		{
-			UE_LOG(LogBuildPatchTool, Error, TEXT("ManifestA/B arguments can't be empty"));
+			UE_LOG(LogBuildPatchTool, Error, TEXT("ManifestA and ManifestB are required parameters"));
 			return false;
 		}
-		
 		NormalizeUriFile(ManifestA);
 		NormalizeUriFile(ManifestB);
 
 		// Get optional values.
 		const bool bHasCloudDir = PARSE_SWITCH(CloudDir);
-		const bool bHasDownloadCloudDirectory = ParseSwitches(TEXT("DownloadCloudDirectory="), DownloadCloudDirectories, Switches);
-
-		PARSE_SWITCH(SourceDetailsLogFilename);
-		PARSE_SWITCH(DeltaFilenameTrailer);
 		PARSE_SWITCH(ScanWindowSize);
 		PARSE_SWITCH(OutputChunkSize);
-		PARSE_SWITCH(MaxRetryCount);
 		if (!bHasCloudDir)
 		{
 			CloudDir = FPaths::GetPath(ManifestB);
 		}
-		if (!bHasDownloadCloudDirectory)
-		{
-			DownloadCloudDirectories.Add(FPaths::GetPath(ManifestB));
-		}
-
 		NormalizeUriPath(CloudDir);
-		for (FString& DownloadCloudDirectory : DownloadCloudDirectories)
+		if (!PARSE_SWITCH(DiffAbortThreshold))
 		{
-			NormalizeUriPath(DownloadCloudDirectory);
+			DiffAbortThreshold = TNumericLimits<int64>::Max();
 		}
 
 		// Clamp ScanWindowSize to sane range.
 		const uint32 RequestedScanWindowSize = ScanWindowSize;
-		ScanWindowSize = FMath::Clamp<uint32>(ScanWindowSize, 128, 128*1024);
+		ScanWindowSize = FMath::Clamp<uint32>(ScanWindowSize, 8000, 128*1024);
 		if (RequestedScanWindowSize != ScanWindowSize)
 		{
-			UE_LOG(LogBuildPatchTool, Warning, TEXT("Requested -ScanWindowSize=%u is outside of allowed range 128KiB >= n >= 128b. Please update your arg to be within range. Continuing with %u."), RequestedScanWindowSize, ScanWindowSize);
+			UE_LOG(LogBuildPatchTool, Warning, TEXT("Requested -ScanWindowSize=%u is outside of allowed range 128KiB >= n >= 8Kb. Please update your arg to be within range. Continuing with %u."), RequestedScanWindowSize, ScanWindowSize);
 		}
 
 		// Clamp OutputChunkSize to sane range.
@@ -204,109 +114,31 @@ private:
 		{
 			UE_LOG(LogBuildPatchTool, Warning, TEXT("Requested -OutputChunkSize=%u is outside of allowed range 10MiB >= n >= 1MB. Please update your arg to be within range. Continuing with %u."), RequestedOutputChunkSize, OutputChunkSize);
 		}
-		DiffAbortThreshold = FDiffAbortThresholdArgument::Parse(*this, Switches);
 
-		// Parse all the secret key and ID pairs.
-		FString ParseErrorMessage;
-		if (!ParsePairs(TEXT("EncryptionSecretId="), TEXT("EncryptionSecretKey="), Switches, EncryptionSecrets, ParseErrorMessage) && !ParseErrorMessage.IsEmpty())
+		// Clamp DiffAbortThreshold to sane min.
+		const uint64 RequestedDiffAbortThreshold = DiffAbortThreshold;
+		DiffAbortThreshold = FMath::Max<uint64>(DiffAbortThreshold, 1000000000);
+		if (RequestedDiffAbortThreshold != DiffAbortThreshold)
 		{
-			UE_LOG(LogBuildPatchTool, Error, TEXT("%s"), *ParseErrorMessage);
-			return false;
-		}
-
-		// Parse any tag aliases we want to use. This is for constructing multiple optimized deltas
-		// where we know we have a larger installation set to draw from instead of just the same tag.
-		// Format is: -ManifestATagAliases="tag1,tag2,tag3,tag4;alias"
-		if (!ParseManyToOneMap(TEXT("ManifestATagAliases"), Switches, ManifestATagAliases))
-		{
-			return false; // already logged.
-		}
-		if (!ParseManyToOneMap(TEXT("ManifestBTagAliases"), Switches, ManifestBTagAliases))
-		{
-			return false; // already logged.
+			UE_LOG(LogBuildPatchTool, Warning, TEXT("Requested -DiffAbortThreshold=%llu is below allowed minimum n >= 1GB. Please update your arg to be above limit. Continuing with %llu."), RequestedDiffAbortThreshold, DiffAbortThreshold);
 		}
 
 		return true;
-#undef SENTENCE_END
-#undef PARSE_UNEXPECTED_SWITCH
-#undef PARSE_REQUIRED_SWITCH
-#undef _PARSE_REQUIRED_SWITCH
 #undef PARSE_SWITCH
 	}
 
-	void LogConfig() const
-	{
-		UE_LOG(LogBuildPatchTool, Log, TEXT("-----Configuration for ChunkDeltaOptimise------"));
-		UE_LOG(LogBuildPatchTool, Log, TEXT("   DiffAbortThreshold: %s"), *LexToString(DiffAbortThreshold));
-		UE_LOG(LogBuildPatchTool, Log, TEXT("   ManifestA: %s"), *ManifestA);
-		UE_LOG(LogBuildPatchTool, Log, TEXT("   ManifestB: %s"), *ManifestB);
-		UE_LOG(LogBuildPatchTool, Log, TEXT("   CloudDir: %s"), *CloudDir);
-		for (const FString& DownloadCloudDirectory : DownloadCloudDirectories)
-		{
-			UE_LOG(LogBuildPatchTool, Log, TEXT("   DownloadCloudDirectory: %s"), *DownloadCloudDirectory);
-		}
-		UE_LOG(LogBuildPatchTool, Log, TEXT("   ScanWindowSize: %u"), ScanWindowSize);
-		UE_LOG(LogBuildPatchTool, Log, TEXT("   OutputChunkSize: %u"), OutputChunkSize);
-		UE_LOG(LogBuildPatchTool, Log, TEXT("   MaxRetryCount: %u"), MaxRetryCount);
-		for (const TPair<FGuid, TArray<uint8>>& EncryptionSecret : EncryptionSecrets)
-		{
-			UE_LOG(LogBuildPatchTool, Log, TEXT("   AvailableSecretId: %s"), *EncryptionSecret.Key.ToString());
-		}
-		if (ManifestATagAliases.Num())
-		{
-			// For logging we reverse the map
-			TMap<FString, TArray<FString>> BackwardsMap;
-			for (const TPair<FString, FString>& Pair : ManifestATagAliases)
-			{
-				BackwardsMap.FindOrAdd(Pair.Value).Add(Pair.Key);
-			}
-			for (const TPair<FString, TArray<FString>>& TagAlias : BackwardsMap)
-			{
-				UE_LOG(LogBuildPatchTool, Log, TEXT("   ManifestATagAliases: %s -> %s"), *FString::Join(TagAlias.Value, TEXT(",")), *TagAlias.Key);
-			}
-		}
-		if (ManifestBTagAliases.Num())
-		{
-			// For logging we reverse the map
-			TMap<FString, TArray<FString>> BackwardsMap;
-			for (const TPair<FString, FString>& Pair : ManifestBTagAliases)
-			{
-				BackwardsMap.FindOrAdd(Pair.Value).Add(Pair.Key);
-			}
-			for (const TPair<FString, TArray<FString>>& TagAlias : BackwardsMap)
-			{
-				UE_LOG(LogBuildPatchTool, Log, TEXT("   ManifestBTagAliases: %s -> %s"), *FString::Join(TagAlias.Value, TEXT(",")), *TagAlias.Key);
-			}
-		}
-		if (DeltaFilenameTrailer.Len())
-		{
-			UE_LOG(LogBuildPatchTool, Log, TEXT("   DeltaFilenameTrailer: %s"), *DeltaFilenameTrailer);
-		}
-		if (SourceDetailsLogFilename.Len())
-		{
-			UE_LOG(LogBuildPatchTool, Log, TEXT("   SourceDetailsLogFilename: %s"), *SourceDetailsLogFilename);
-		}
-	}
-
 private:
-	static TCHAR const* const MODE_NAME;
 	IBuildPatchServicesModule& BpsInterface;
-	const TCHAR* CommandLine;
-	bool bHelp = false;
+	bool bHelp;
 	FString ManifestA;
 	FString ManifestB;
 	FString CloudDir;
-	TArray<FString> DownloadCloudDirectories;
-	uint32 ScanWindowSize = 8191U;
-	uint32 OutputChunkSize = 1024U * 1024U;
-	uint32 MaxRetryCount = 30;
-	TOptional<BuildPatchServices::FDiffAbortThreshold> DiffAbortThreshold;
-	TMap<FString, FString> ManifestBTagAliases;
-	TMap<FString, FString> ManifestATagAliases;
-	FString DeltaFilenameTrailer;
-	FString SourceDetailsLogFilename;
-	TMap<FGuid, TArray<uint8>> EncryptionSecrets;
+	uint32 ScanWindowSize;
+	uint32 OutputChunkSize;
+	uint64 DiffAbortThreshold;
 };
 
-IMPLEMENT_BPT_MODE(ChunkDeltaOptimise, FChunkDeltaOptimiseMode);
-IMPLEMENT_BPT_MODE_ALIAS(ChunkDeltaOptimize, FChunkDeltaOptimiseMode);
\ No newline at end of file
+BuildPatchTool::IToolModeRef BuildPatchTool::FChunkDeltaOptimiseToolModeFactory::Create(IBuildPatchServicesModule& BpsInterface)
+{
+	return MakeShareable(new FChunkDeltaOptimiseMode(BpsInterface));
+}
diff --git a/Engine/Source/Programs/BuildPatchTool/Private/ToolModes/CompactifyMode.cpp b/Engine/Source/Programs/BuildPatchTool/Private/ToolModes/CompactifyMode.cpp
--- a/Engine/Source/Programs/BuildPatchTool/Private/ToolModes/CompactifyMode.cpp
+++ b/Engine/Source/Programs/BuildPatchTool/Private/ToolModes/CompactifyMode.cpp
@@ -1,22 +1,19 @@
 // Copyright Epic Games, Inc. All Rights Reserved.
 
 #include "ToolModes/CompactifyMode.h"
-
-#include "BuildPatchSettings.h"
-#include "BuildPatchTool.h"
 #include "Interfaces/IBuildPatchServicesModule.h"
+#include "BuildPatchTool.h"
 #include "Misc/CommandLine.h"
 #include "Misc/Paths.h"
-#include "ToolModes/ToolModesHelp.h"
+#include "BuildPatchSettings.h"
 
 using namespace BuildPatchTool;
 
 class FCompactifyToolMode : public IToolMode
 {
 public:
-	FCompactifyToolMode(IBuildPatchServicesModule& InBpsInterface, const TCHAR* InCommandLine)
+	FCompactifyToolMode(IBuildPatchServicesModule& InBpsInterface)
 		: BpsInterface(InBpsInterface)
-		, CommandLine(InCommandLine)
 	{}
 
 	virtual ~FCompactifyToolMode()
@@ -33,12 +30,23 @@ public:
 		// Print help if requested
 		if (bHelp)
 		{
-			PrintHelp<FCompactifyToolModeHelp>();
+			UE_LOG(LogBuildPatchTool, Display, TEXT("COMPACTIFY MODE"));
+			UE_LOG(LogBuildPatchTool, Display, TEXT("This tool supports the removal of redundant patch data from a cloud directory."));
+			UE_LOG(LogBuildPatchTool, Display, TEXT(""));
+			UE_LOG(LogBuildPatchTool, Display, TEXT("Required arguments:"));
+			UE_LOG(LogBuildPatchTool, Display, TEXT("  -mode=Compactify           Must be specified to launch the tool in compactify mode."));
+			UE_LOG(LogBuildPatchTool, Display, TEXT("  -CloudDir=\"\"               Specifies in quotes the cloud directory where manifest files and chunks to be compactified can be found."));
+			UE_LOG(LogBuildPatchTool, Display, TEXT("  -DataAgeThreshold=2        The minimum age in days of chunk files that will be deleted. Any unreferenced chunks older than this will be deleted."));
+			UE_LOG(LogBuildPatchTool, Display, TEXT(""));
+			UE_LOG(LogBuildPatchTool, Display, TEXT("Optional arguments:"));
+			UE_LOG(LogBuildPatchTool, Display, TEXT("  -preview                   Log all the actions it will take to update internal structures, but don't actually execute them."));
+			UE_LOG(LogBuildPatchTool, Display, TEXT("  -DeletedChunkLogFile=\"\"    Log the list of paths of deleted chunk files to this specified filename. All paths are relative to CloudDir."));
+			UE_LOG(LogBuildPatchTool, Display, TEXT(""));
+			UE_LOG(LogBuildPatchTool, Display, TEXT("NB: If -DataAgeThreshold is not supplied, then all unreferenced existing data is eligible for deletion by the compactify process."));
+			UE_LOG(LogBuildPatchTool, Display, TEXT(""));
 			return EReturnCode::OK;
 		}
 
-		LogConfig();
-
 		// Grab options
 		BuildPatchServices::FCompactifyConfiguration Configuration;
 		Configuration.CloudDirectory = CloudDir;
@@ -57,7 +65,7 @@ private:
 	{
 #define PARSE_SWITCH(Switch) ParseSwitch(TEXT(#Switch "="), Switch, Switches)
 		TArray<FString> Tokens, Switches;
-		FCommandLine::Parse(CommandLine, Tokens, Switches);
+		FCommandLine::Parse(FCommandLine::Get(), Tokens, Switches);
 
 		bHelp = ParseOption(TEXT("help"), Switches);
 		if (bHelp)
@@ -72,17 +80,10 @@ private:
 			UE_LOG(LogBuildPatchTool, Error, TEXT("CloudDir and DataAgeThreshold are required parameters"));
 			return false;
 		}
-
-		if (CloudDir.IsEmpty())
-		{
-			UE_LOG(LogBuildPatchTool, Error, TEXT("CloudDir argument can't be empty"));
-			return false;
-		}
-
 		NormalizeUriPath(CloudDir);
 
 		// Check required numeric values
-		if (!DataAgeThreshold.IsNumeric() || TCString<TCHAR>::Atod(*DataAgeThreshold) < 0.0)
+		if (!DataAgeThreshold.IsNumeric())
 		{
 			UE_LOG(LogBuildPatchTool, Error, TEXT("An error occurred processing numeric token from commandline -DataAgeThreshold=%s"), *DataAgeThreshold);
 			return false;
@@ -97,19 +98,8 @@ private:
 #undef PARSE_SWITCH
 	}
 
-	void LogConfig() const
-	{
-		UE_LOG(LogBuildPatchTool, Log, TEXT("-----Configuration for CompactifyMode------"));
-		UE_LOG(LogBuildPatchTool, Log, TEXT("   CloudDir: %s"), *CloudDir);
-		UE_LOG(LogBuildPatchTool, Log, TEXT("   DataAgeThreshold: %s"), *DataAgeThreshold);
-		UE_LOG(LogBuildPatchTool, Log, TEXT("   DeletedChunkLogFile: %s"), *DeletedChunkLogFile);
-		UE_LOG(LogBuildPatchTool, Log, TEXT("   bPreview: %s"), bPreview ? TEXT("true") : TEXT("false"));
-	}
-
 private:
-	static TCHAR const* const MODE_NAME;
 	IBuildPatchServicesModule& BpsInterface;
-	const TCHAR* CommandLine;
 	bool bHelp;
 	FString CloudDir;
 	FString DataAgeThreshold;
@@ -117,4 +107,7 @@ private:
 	bool bPreview;
 };
 
-IMPLEMENT_BPT_MODE(Compactify, FCompactifyToolMode);
+BuildPatchTool::IToolModeRef BuildPatchTool::FCompactifyToolModeFactory::Create(IBuildPatchServicesModule& BpsInterface)
+{
+	return MakeShareable(new FCompactifyToolMode(BpsInterface));
+}
diff --git a/Engine/Source/Programs/BuildPatchTool/Private/ToolModes/DiffManifestMode.cpp b/Engine/Source/Programs/BuildPatchTool/Private/ToolModes/DiffManifestMode.cpp
--- a/Engine/Source/Programs/BuildPatchTool/Private/ToolModes/DiffManifestMode.cpp
+++ b/Engine/Source/Programs/BuildPatchTool/Private/ToolModes/DiffManifestMode.cpp
@@ -1,22 +1,18 @@
 // Copyright Epic Games, Inc. All Rights Reserved.
 
 #include "ToolModes/DiffManifestMode.h"
-
-#include "BuildPatchTool.h"
 #include "Interfaces/IBuildPatchServicesModule.h"
-#include "Interfaces/ToolMode.h"
+#include "BuildPatchTool.h"
 #include "Misc/CommandLine.h"
 #include "Misc/Paths.h"
-#include "ToolModes/ToolModesHelp.h"
 
 using namespace BuildPatchTool;
 
 class FDiffManifestToolMode : public IToolMode
 {
 public:
-	FDiffManifestToolMode(IBuildPatchServicesModule& InBpsInterface, const TCHAR* InCommandLine)
+	FDiffManifestToolMode(IBuildPatchServicesModule& InBpsInterface)
 		: BpsInterface(InBpsInterface)
-		, CommandLine(InCommandLine)
 	{}
 
 	virtual ~FDiffManifestToolMode()
@@ -33,32 +29,60 @@ public:
 		// Print help if requested.
 		if (bHelp)
 		{
-			PrintHelp<FDiffManifestToolModeHelp>();
+			UE_LOG(LogBuildPatchTool, Display, TEXT("DIFF MANIFEST MODE"));
+			UE_LOG(LogBuildPatchTool, Display, TEXT("This tool mode reports the changes between two existing manifest files."));
+			UE_LOG(LogBuildPatchTool, Display, TEXT(""));
+			UE_LOG(LogBuildPatchTool, Display, TEXT("Required arguments:"));
+			UE_LOG(LogBuildPatchTool, Display, TEXT("  -mode=DiffManifests    Must be specified to launch the tool in diff manifests mode."));
+			UE_LOG(LogBuildPatchTool, Display, TEXT("  -ManifestA=\"\"          Specifies in quotes the file path to the base manifest."));
+			UE_LOG(LogBuildPatchTool, Display, TEXT("  -ManifestB=\"\"          Specifies in quotes the file path to the update manifest."));
+			UE_LOG(LogBuildPatchTool, Display, TEXT(""));
+			UE_LOG(LogBuildPatchTool, Display, TEXT("Optional arguments:"));
+			UE_LOG(LogBuildPatchTool, Display, TEXT("  -InstallTagsA=\"\"       Specifies in quotes a comma seperated list of install tags used on ManifestA. You should include empty string if you want to count untagged files."));
+			UE_LOG(LogBuildPatchTool, Display, TEXT("                           Leaving the parameter out will use all files."));
+			UE_LOG(LogBuildPatchTool, Display, TEXT("                           -InstallTagsA=\"\" will be untagged files only."));
+			UE_LOG(LogBuildPatchTool, Display, TEXT("                           -InstallTagsA=\",tag\" will be untagged files plus files tagged with 'tag'."));
+			UE_LOG(LogBuildPatchTool, Display, TEXT("                           -InstallTagsA=\"tag\" will be files tagged with 'tag' only."));
+			UE_LOG(LogBuildPatchTool, Display, TEXT("  -InstallTagsB=\"\"       Specifies in quotes a comma seperated list of install tags used on ManifestB. Same rules apply as InstallTagsA."));
+			UE_LOG(LogBuildPatchTool, Display, TEXT("  -CompareTagSet=\"\"      Specifies in quotes a comma seperated list of install tags used to calculate differential statistics betweeen the manifests. Multiple lists are allowed. Same rules apply as InstallTagsA."));
+			UE_LOG(LogBuildPatchTool, Display, TEXT("  -OutputFile=\"\"         Specifies in quotes the file path where the diff will be exported as a JSON object."));
+			UE_LOG(LogBuildPatchTool, Display, TEXT("  -PatchDescriptorPath=\"\" If specified, will write out metadata files specifying what parts of each patched file changed."));
+			UE_LOG(LogBuildPatchTool, Display, TEXT("  -OnlyPatchDescriptors    If specified with PatchDescriptorPath, DiffManifests will only emit the files and not emit diff info."));
+			UE_LOG(LogBuildPatchTool, Display, TEXT("  -RequireOptimizedDelta   If specified, fail if an optimized delta file doesn't exist."));
+			UE_LOG(LogBuildPatchTool, Display, TEXT(""));
 			return EReturnCode::OK;
 		}
 
-		LogConfig();
+		// Calc desired tags.
+		TSet<FString> TagSetA, TagSetB;
+		if (bHasTagsA)
+		{
+			TagSetA = ProcessTagList(InstallTagsA);
+		}
+		if (bHasTagsB)
+		{
+			TagSetB = ProcessTagList(InstallTagsB);
+		}
+
+		TArray<TSet<FString>> CompareTagSets;
 
-		// Register secrets.
-		for (const TPair<FGuid, TArray<uint8>>& IdKeyPair : EncryptionSecrets)
+		for (const FString& CompareTagsList : CompareTagsArray)
 		{
-			BpsInterface.RegisterEncryptionSecret(IdKeyPair.Key, IdKeyPair.Value);
+			CompareTagSets.Add(ProcessTagList(CompareTagsList));
 		}
 
 		// Setup and run.
 		BuildPatchServices::FDiffManifestsConfiguration Configuration;
 		Configuration.ManifestAUri = ManifestA;
 		Configuration.ManifestBUri = ManifestB;
-		Configuration.TagSetA = InstallTagsA;
-		Configuration.TagSetB = InstallTagsB;
-		Configuration.CompareTagSetsA = CompareTagSets;
-		Configuration.CompareTagSetsB = CompareNewTagSets;
+		Configuration.TagSetA = TagSetA;
+		Configuration.TagSetB = TagSetB;
+		Configuration.CompareTagSets = CompareTagSets;
 		Configuration.OutputFilePath = OutputFile;
 		Configuration.OutputPatchDescriptorPath = PatchDescriptorPath;
 		Configuration.bOnlyPatchDescriptors = bOnlyPatchDescriptors;
 		Configuration.bRequireOptimizedDelta = bRequireOptimizedDelta;
-		Configuration.DeltaFilenameTrailer = DeltaFilenameTrailer;
-		Configuration.bSimulateInstallTimes = bSimulateInstallTimes;
+		Configuration.bEmitInstallTime = !bSkipInstallTime;
 
 		// Run the merge manifest routine.
 		bool bSuccess = BpsInterface.DiffManifests(Configuration);
@@ -70,17 +94,8 @@ private:
 	bool ProcessCommandLine()
 	{
 #define PARSE_SWITCH(Switch) ParseSwitch(TEXT(#Switch "="), Switch, Switches)
-#define _PARSE_REQUIRED_SWITCH(Switch, bExpect) \
-		if (PARSE_SWITCH(Switch) != bExpect) \
-		{ \
-			UE_LOG(LogBuildPatchTool, Error, #Switch SENTENCE_END); \
-			bParametersOk = false; \
-		}
-#define PARSE_REQUIRED_SWITCH(Switch) _PARSE_REQUIRED_SWITCH(Switch, true)
-#define PARSE_UNEXPECTED_SWITCH(Switch) _PARSE_REQUIRED_SWITCH(Switch, false)
-
 		TArray<FString> Tokens, Switches;
-		FCommandLine::Parse(CommandLine, Tokens, Switches);
+		FCommandLine::Parse(FCommandLine::Get(), Tokens, Switches);
 
 		bHelp = ParseOption(TEXT("help"), Switches);
 		if (bHelp)
@@ -88,85 +103,32 @@ private:
 			return true;
 		}
 
-		// local variables to check for unexpected parameters
-		FString BuildVersionA, BuildVersionB, ClientId, OrganizationId, ProductId;
-#define SENTENCE_END TEXT(" is a required parameter for offline mode.")
-		bool bParametersOk = true;
-		PARSE_REQUIRED_SWITCH(ManifestA);
-		PARSE_REQUIRED_SWITCH(ManifestB);
-#undef SENTENCE_END
-#define SENTENCE_END TEXT(" should only be used when uploading to Epic services.")
-		PARSE_UNEXPECTED_SWITCH(BuildVersionA);
-		PARSE_UNEXPECTED_SWITCH(BuildVersionB);
-#undef SENTENCE_END
-#define SENTENCE_END TEXT(" is only supported when Epic services are available. Download an official build from Epic to enable this functionality.")
-		PARSE_UNEXPECTED_SWITCH(ClientId);
-		PARSE_UNEXPECTED_SWITCH(OrganizationId);
-		PARSE_UNEXPECTED_SWITCH(ProductId);
-#undef SENTENCE_END
-
-		if (!bParametersOk)
+		// Get all required parameters.
+		if (!(PARSE_SWITCH(ManifestA)
+		   && PARSE_SWITCH(ManifestB)))
 		{
+			UE_LOG(LogBuildPatchTool, Error, TEXT("ManifestA and ManifestB are required parameters."));
 			return false;
 		}
-
-		if (ManifestA.IsEmpty() || ManifestB.IsEmpty())
-		{
-			UE_LOG(LogBuildPatchTool, Error, TEXT("ManifestA/B arguments can't be empty"));
-			return false;
-		}
-		
 		NormalizeUriFile(ManifestA);
 		NormalizeUriFile(ManifestB);
 
 		// Get optional parameters
-		PARSE_SWITCH(InstallTagsA);
-		PARSE_SWITCH(InstallTagsB);
+		bHasTagsA = PARSE_SWITCH(InstallTagsA);
+		bHasTagsB = PARSE_SWITCH(InstallTagsB);
 		PARSE_SWITCH(OutputFile);
 		NormalizeUriFile(OutputFile);
 
+		ParseSwitches(TEXT("CompareTagSet="), CompareTagsArray, Switches);
+
 		PARSE_SWITCH(PatchDescriptorPath);
 		NormalizeUriFile(PatchDescriptorPath);
 
-		PARSE_SWITCH(DeltaFilenameTrailer);
-		
 		bRequireOptimizedDelta = ParseOption(TEXT("RequireOptimizedDelta"), Switches);
 		bOnlyPatchDescriptors = ParseOption(TEXT("OnlyPatchDescriptors"), Switches);
-		bSimulateInstallTimes = ParseOption(TEXT("SimulateInstallTimes"), Switches);
-		bool bSkipInstallTime = ParseOption(TEXT("SkipInstallTime"), Switches);
-		if (bSkipInstallTime)
-		{
-			bSimulateInstallTimes = false;
-		}
-
-		TArray<FString> CompareTagsArray;
-		ParseSwitches(TEXT("CompareTagSet="), CompareTagsArray, Switches);
-
-		for (const FString& CompareTagsList : CompareTagsArray)
-		{
-			CompareTagSets.Add(ProcessTagList(CompareTagsList));
-		}
-
-		TArray<FString> CompareNewTagsArray;
-		ParseSwitches(TEXT("CompareNewTagSet="), CompareNewTagsArray, Switches);
-
-		for (const FString& CompareTagsList : CompareNewTagsArray)
-		{
-			CompareNewTagSets.Add(ProcessTagList(CompareTagsList));
-		}
-
-		// Parse all the secret key and ID pairs.
-		FString ParseErrorMessage;
-		if (!ParsePairs(TEXT("EncryptionSecretId="), TEXT("EncryptionSecretKey="), Switches, EncryptionSecrets, ParseErrorMessage) && !ParseErrorMessage.IsEmpty())
-		{
-			UE_LOG(LogBuildPatchTool, Error, TEXT("%s"), *ParseErrorMessage);
-			return false;
-		}
+		bSkipInstallTime = ParseOption(TEXT("SkipInstallTime"), Switches);
 
 		return true;
-#undef PARSE_UNEXPECTED_SWITCH
-#undef PARSE_REQUIRED_SWITCH
-#undef _PARSE_REQUIRED_SWITCH
 #undef PARSE_SWITCH
 	}
 
@@ -185,56 +147,24 @@ private:
 		return TSet<FString>(MoveTemp(TagArray));
 	}
 
-	void LogConfig() const
-	{
-		UE_LOG(LogBuildPatchTool, Log, TEXT("-----Configuration for DiffManifests------"));
-		UE_LOG(LogBuildPatchTool, Log, TEXT("   ManifestA: %s"), *ManifestA);
-		UE_LOG(LogBuildPatchTool, Log, TEXT("   ManifestB: %s"), *ManifestB);
-		TArray<FString> CompareTagSetStrings;
-		for (const TSet<FString>& CompareTagSet : CompareTagSets)
-		{
-			CompareTagSetStrings.Add(FString::Join(CompareTagSet, TEXT(",")));
-		}
-		TArray<FString> CompareNewTagSetStrings;
-		for (const TSet<FString>& CompareTagSet : CompareNewTagSets)
-		{
-			CompareNewTagSetStrings.Add(FString::Join(CompareTagSet, TEXT(",")));
-		}
-		UE_LOG(LogBuildPatchTool, Log, TEXT("   CompareTagSets: %s"), *FString::Join(CompareTagSetStrings, TEXT(";")));
-		UE_LOG(LogBuildPatchTool, Log, TEXT("   CompareNewTagSets: %s"), *FString::Join(CompareNewTagSetStrings, TEXT(";")));
-		UE_LOG(LogBuildPatchTool, Log, TEXT("   InstallTagsA: %s"), *FString::Join(InstallTagsA.Array(), TEXT(",")));
-		UE_LOG(LogBuildPatchTool, Log, TEXT("   InstallTagsB: %s"), *FString::Join(InstallTagsB.Array(), TEXT(",")));
-		UE_LOG(LogBuildPatchTool, Log, TEXT("   OutputFile: %s"), *OutputFile);
-		UE_LOG(LogBuildPatchTool, Log, TEXT("   RequireOptimizedDelta: %s"), *LexToString(bRequireOptimizedDelta));
-		UE_LOG(LogBuildPatchTool, Log, TEXT("   SimulateInstallTimes: %s"), *LexToString(bSimulateInstallTimes));
-		if (DeltaFilenameTrailer.Len())
-		{
-			UE_LOG(LogBuildPatchTool, Log, TEXT("   DeltaFilenameTrailer: %s"), *DeltaFilenameTrailer);		
-		}
-		for (const TPair<FGuid, TArray<uint8>>& EncryptionSecret : EncryptionSecrets)
-		{
-			UE_LOG(LogBuildPatchTool, Log, TEXT("   AvailableSecretId: %s"), *EncryptionSecret.Key.ToString());
-		}
-	}
-
 private:
-	static TCHAR const* const MODE_NAME;
 	IBuildPatchServicesModule& BpsInterface;
-	const TCHAR* CommandLine;
 	bool bHelp;
 	FString ManifestA;
 	FString ManifestB;
-	TArray<TSet<FString>> CompareTagSets;
-	TArray<TSet<FString>> CompareNewTagSets;
-	TSet<FString> InstallTagsA;
-	TSet<FString> InstallTagsB;
+	TArray<FString> CompareTagsArray;
+	bool bHasTagsA;
+	bool bHasTagsB;
+	FString InstallTagsA;
+	FString InstallTagsB;
 	FString OutputFile;
 	FString PatchDescriptorPath;
 	bool bRequireOptimizedDelta;
 	bool bOnlyPatchDescriptors;
-	bool bSimulateInstallTimes;
-	FString DeltaFilenameTrailer;
-	TMap<FGuid, TArray<uint8>> EncryptionSecrets;
+	bool bSkipInstallTime;
 };
 
-IMPLEMENT_BPT_MODE(DiffManifests, FDiffManifestToolMode);
+BuildPatchTool::IToolModeRef BuildPatchTool::FDiffManifestToolModeFactory::Create(IBuildPatchServicesModule& BpsInterface)
+{
+	return MakeShareable(new FDiffManifestToolMode(BpsInterface));
+}
diff --git a/Engine/Source/Programs/BuildPatchTool/Private/ToolModes/EnumerationMode.cpp b/Engine/Source/Programs/BuildPatchTool/Private/ToolModes/EnumerationMode.cpp
--- a/Engine/Source/Programs/BuildPatchTool/Private/ToolModes/EnumerationMode.cpp
+++ b/Engine/Source/Programs/BuildPatchTool/Private/ToolModes/EnumerationMode.cpp
@@ -1,21 +1,18 @@
 // Copyright Epic Games, Inc. All Rights Reserved.
 
 #include "ToolModes/EnumerationMode.h"
-
-#include "BuildPatchTool.h"
 #include "Interfaces/IBuildPatchServicesModule.h"
+#include "BuildPatchTool.h"
 #include "Misc/CommandLine.h"
 #include "Misc/Paths.h"
-#include "ToolModes/ToolModesHelp.h"
 
 using namespace BuildPatchTool;
 
 class FEnumerationToolMode : public IToolMode
 {
 public:
-	FEnumerationToolMode(IBuildPatchServicesModule& InBpsInterface, const TCHAR* InCommandLine)
+	FEnumerationToolMode(IBuildPatchServicesModule& InBpsInterface)
 		: BpsInterface(InBpsInterface)
-		, CommandLine(InCommandLine)
 	{}
 
 	virtual ~FEnumerationToolMode()
@@ -23,8 +20,6 @@ public:
 
 	virtual EReturnCode Execute() override
 	{
-		using namespace BuildPatchServices;
-
 		// Parse commandline
 		if (ProcessCommandline() == false)
 		{
@@ -34,34 +29,35 @@ public:
 		// Print help if requested
 		if (bHelp)
 		{
-			PrintHelp<FEnumerationToolModeHelp>();
+			UE_LOG(LogBuildPatchTool, Display, TEXT("ENUMERATION MODE"));
+			UE_LOG(LogBuildPatchTool, Display, TEXT("This tool supports enumerating patch data referenced by a build manifest or chunkdb."));
+			UE_LOG(LogBuildPatchTool, Display, TEXT(""));
+			UE_LOG(LogBuildPatchTool, Display, TEXT("Required arguments:"));
+			UE_LOG(LogBuildPatchTool, Display, TEXT("  -mode=Enumeration    Must be specified to launch the tool in enumeration mode."));
+			UE_LOG(LogBuildPatchTool, Display, TEXT("  -InputFile=\"\"        Specifies in quotes the file path to the manifest or chunkdb to enumerate from."));
+			UE_LOG(LogBuildPatchTool, Display, TEXT("  -OutputFile=\"\"       Specifies in quotes the file path to a file where the list will be saved out."));
+			UE_LOG(LogBuildPatchTool, Display, TEXT(""));
+			UE_LOG(LogBuildPatchTool, Display, TEXT("Optional arguments:"));
+			UE_LOG(LogBuildPatchTool, Display, TEXT("  -includesizes        When specified, the size of each file in bytes will also be output (see notes)."));
+			UE_LOG(LogBuildPatchTool, Display, TEXT(""));
+			UE_LOG(LogBuildPatchTool, Display, TEXT("NB: If InputFile is a manifest, the output file format will be text file with one line per file, each containing cloud relative path."));
+			UE_LOG(LogBuildPatchTool, Display, TEXT("    e.g. path/to/file"));
+			UE_LOG(LogBuildPatchTool, Display, TEXT("    If InputFile is a chunk package, the output file format will be text file with one line per chunk, each containing tab separated hex chunk"));
+			UE_LOG(LogBuildPatchTool, Display, TEXT("    GUID, hex chunk rolling hash, and hex chunk SHA1."));
+			UE_LOG(LogBuildPatchTool, Display, TEXT("    e.g. 2CC26D05B64363780D5CF292E6B570A3\\t078070129133079067060057\\t527490FCA1DA6FAAB0E6F6E369E372FA693CCFBB"));
+			UE_LOG(LogBuildPatchTool, Display, TEXT("    Line endings are \\r\\n."));
+			UE_LOG(LogBuildPatchTool, Display, TEXT("NB: If includesizes is specified, each line of the output text file will end with a tab then the number of bytes of the file."));
 			return EReturnCode::OK;
 		}
 
-		LogConfig();
-
-		// Register secrets.
-		for (const TPair<FGuid, TArray<uint8>>& IdKeyPair : EncryptionSecrets)
-		{
-			BpsInterface.RegisterEncryptionSecret(IdKeyPair.Key, IdKeyPair.Value);
-		}
-
-		// Check existence of input file
-		if (!FPaths::FileExists(InputFile))
-		{
-			UE_LOG(LogBuildPatchTool, Error, TEXT("Provided input file was not found: %s"), *InputFile);
-			return EReturnCode::FileNotFound;
-		}
-
 		// Setup and run.
-		FPatchDataEnumerationConfiguration Settings;
+		BuildPatchServices::FPatchDataEnumerationConfiguration Settings;
 		Settings.InputFile = InputFile;
 		Settings.OutputFile = OutputFile;
 		Settings.bIncludeSizes = bIncludeSizes;
 
 		// Run the enumeration routine
-		IPatchDataEnumerationRef PatchDataEnumeration = BpsInterface.CreatePatchDataEnumeration(Settings);
-		const bool bSuccess = PatchDataEnumeration->Run();
+		bool bSuccess = BpsInterface.EnumeratePatchData(Settings);
 		return bSuccess ? EReturnCode::OK : EReturnCode::ToolFailure;
 	}
 
@@ -71,7 +67,7 @@ private:
 	{
 #define PARSE_SWITCH(Switch) ParseSwitch(TEXT(#Switch "="), Switch, Switches)
 		TArray<FString> Tokens, Switches;
-		FCommandLine::Parse(CommandLine, Tokens, Switches);
+		FCommandLine::Parse(FCommandLine::Get(), Tokens, Switches);
 
 		bHelp = ParseOption(TEXT("help"), Switches);
 		if (bHelp)
@@ -79,59 +75,39 @@ private:
 			return true;
 		}
 
-		// Get all required parameters
-		if (!PARSE_SWITCH(InputFile) || InputFile.IsEmpty()
-		  || !PARSE_SWITCH(OutputFile) || OutputFile.IsEmpty())
+		// Check for deprecated ManifestFile param.
+		FString ManifestFile;
+		if (PARSE_SWITCH(ManifestFile))
 		{
-			UE_LOG(LogBuildPatchTool, Error, TEXT("InputFile and OutputFile are required parameters"));
-			return false;
+			InputFile = ManifestFile;
 		}
 
-		if (InputFile.Equals(OutputFile, ESearchCase::IgnoreCase))
+		// Get all required parameters
+		if ((!PARSE_SWITCH(InputFile) && InputFile.IsEmpty())
+		  || !PARSE_SWITCH(OutputFile))
 		{
-			UE_LOG(LogBuildPatchTool, Error, TEXT("InputFile and OutputFile cannot be the same"));
+			UE_LOG(LogBuildPatchTool, Error, TEXT("InputFile and OutputFile are required parameters"));
 			return false;
 		}
-
 		NormalizeUriFile(InputFile);
 		NormalizeUriFile(OutputFile);
 
 		// Get optional parameters
 		bIncludeSizes = ParseOption(TEXT("includesizes"), Switches);
 
-		// Parse all the secret key and ID pairs.
-		FString ParseErrorMessage;
-		if (!ParsePairs(TEXT("EncryptionSecretId="), TEXT("EncryptionSecretKey="), Switches, EncryptionSecrets, ParseErrorMessage) && !ParseErrorMessage.IsEmpty())
-		{
-			UE_LOG(LogBuildPatchTool, Error, TEXT("%s"), *ParseErrorMessage);
-			return false;
-		}
-
 		return true;
 #undef PARSE_SWITCH
 	}
 
-	void LogConfig() const
-	{
-		UE_LOG(LogBuildPatchTool, Log, TEXT("-----Configuration for EnumerationMode------"));
-		UE_LOG(LogBuildPatchTool, Log, TEXT("   InputFile: %s"), *InputFile);
-		UE_LOG(LogBuildPatchTool, Log, TEXT("   OutputFile: %s"), *OutputFile);
-		UE_LOG(LogBuildPatchTool, Log, TEXT("   bIncludeSizes: %s"), bIncludeSizes ? TEXT("true") : TEXT("false"));
-		for (const TPair<FGuid, TArray<uint8>>& EncryptionSecret : EncryptionSecrets)
-		{
-			UE_LOG(LogBuildPatchTool, Log, TEXT("   AvailableSecretId: %s"), *EncryptionSecret.Key.ToString());
-		}
-	}
-
 private:
-	static TCHAR const* const MODE_NAME;
 	IBuildPatchServicesModule& BpsInterface;
-	const TCHAR* CommandLine;
 	bool bHelp;
 	FString InputFile;
 	FString OutputFile;
 	bool bIncludeSizes;
-	TMap<FGuid, TArray<uint8>> EncryptionSecrets;
 };
 
-IMPLEMENT_BPT_MODE(Enumeration, FEnumerationToolMode);
+BuildPatchTool::IToolModeRef BuildPatchTool::FEnumerationToolModeFactory::Create(IBuildPatchServicesModule& BpsInterface)
+{
+	return MakeShareable(new FEnumerationToolMode(BpsInterface));
+}
diff --git a/Engine/Source/Programs/BuildPatchTool/Private/ToolModes/MergeManifestMode.cpp b/Engine/Source/Programs/BuildPatchTool/Private/ToolModes/MergeManifestMode.cpp
--- a/Engine/Source/Programs/BuildPatchTool/Private/ToolModes/MergeManifestMode.cpp
+++ b/Engine/Source/Programs/BuildPatchTool/Private/ToolModes/MergeManifestMode.cpp
@@ -1,21 +1,18 @@
 // Copyright Epic Games, Inc. All Rights Reserved.
 
 #include "ToolModes/MergeManifestMode.h"
-
-#include "BuildPatchTool.h"
 #include "Interfaces/IBuildPatchServicesModule.h"
+#include "BuildPatchTool.h"
 #include "Misc/CommandLine.h"
 #include "Misc/Paths.h"
-#include "ToolModes/ToolModesHelp.h"
 
 using namespace BuildPatchTool;
 
 class FMergeManifestToolMode : public IToolMode
 {
 public:
-	FMergeManifestToolMode(IBuildPatchServicesModule& InBpsInterface, const TCHAR* InCommandLine)
+	FMergeManifestToolMode(IBuildPatchServicesModule& InBpsInterface)
 		: BpsInterface(InBpsInterface)
-		, CommandLine(InCommandLine)
 	{}
 
 	virtual ~FMergeManifestToolMode()
@@ -32,18 +29,25 @@ public:
 		// Print help if requested
 		if (bHelp)
 		{
-			PrintHelp<FMergeManifestToolModeHelp>();
+			UE_LOG(LogBuildPatchTool, Display, TEXT("MERGE MANIFEST MODE"));
+			UE_LOG(LogBuildPatchTool, Display, TEXT("This tool supports generating a hotfix manifest from two existing manifest files."));
+			UE_LOG(LogBuildPatchTool, Display, TEXT(""));
+			UE_LOG(LogBuildPatchTool, Display, TEXT("Required arguments:"));
+			UE_LOG(LogBuildPatchTool, Display, TEXT("  -mode=MergeManifests    Must be specified to launch the tool in merge manifests mode."));
+			UE_LOG(LogBuildPatchTool, Display, TEXT("  -ManifestA=\"\"           Specifies in quotes the file path to the base manifest."));
+			UE_LOG(LogBuildPatchTool, Display, TEXT("  -ManifestB=\"\"           Specifies in quotes the file path to the update manifest."));
+			UE_LOG(LogBuildPatchTool, Display, TEXT("  -ManifestC=\"\"           Specifies in quotes the file path to the output manifest."));
+			UE_LOG(LogBuildPatchTool, Display, TEXT("  -BuildVersion=\"\"        Specifies in quotes the new version string for the build being produced."));
+			UE_LOG(LogBuildPatchTool, Display, TEXT(""));
+			UE_LOG(LogBuildPatchTool, Display, TEXT("Optional arguments:"));
+			UE_LOG(LogBuildPatchTool, Display, TEXT("  -MergeFileList=\"\"       Specifies in quotes, the path to a text file containing complete list of desired build root relative files followed by \\t character, followed by A or B to select the manifest to pull from. These should be seperated by \\r\\n line endings."));
+			UE_LOG(LogBuildPatchTool, Display, TEXT(""));
+			UE_LOG(LogBuildPatchTool, Display, TEXT("NB: If -MergeFileList is not specified, then union of all files will be selected, preferring ManifestB's version."));
+			UE_LOG(LogBuildPatchTool, Display, TEXT("NB: With the exception of the new version string for the build, all meta will be copied from only ManifestB."));
+			UE_LOG(LogBuildPatchTool, Display, TEXT(""));
 			return EReturnCode::OK;
 		}
 
-		LogConfig();
-
-		// Register secrets.
-		for (const TPair<FGuid, TArray<uint8>>& IdKeyPair : EncryptionSecrets)
-		{
-			BpsInterface.RegisterEncryptionSecret(IdKeyPair.Key, IdKeyPair.Value);
-		}
-
 		// Run the merge manifest routine
 		bool bSuccess = BpsInterface.MergeManifests(ManifestA, ManifestB, ManifestC, BuildVersion, MergeFileList);
 		return bSuccess ? EReturnCode::OK : EReturnCode::ToolFailure;
@@ -55,7 +59,7 @@ private:
 	{
 #define PARSE_SWITCH(Switch) ParseSwitch(TEXT(#Switch "="), Switch, Switches)
 		TArray<FString> Tokens, Switches;
-		FCommandLine::Parse(CommandLine, Tokens, Switches);
+		FCommandLine::Parse(FCommandLine::Get(), Tokens, Switches);
 
 		bHelp = ParseOption(TEXT("help"), Switches);
 		if (bHelp)
@@ -72,64 +76,29 @@ private:
 			UE_LOG(LogBuildPatchTool, Error, TEXT("ManifestA, ManifestB, ManifestC, and BuildVersion are required parameters"));
 			return false;
 		}
-
-		if (ManifestA.IsEmpty() || ManifestB.IsEmpty() || ManifestC.IsEmpty() || BuildVersion.IsEmpty())
-		{
-			UE_LOG(LogBuildPatchTool, Error, TEXT("ManifestA/B/C and BuildVersion arguments can't be empty"));
-			return false;
-		}
-
 		NormalizeUriFile(ManifestA);
 		NormalizeUriFile(ManifestB);
 		NormalizeUriFile(ManifestC);
 
-		if (ManifestC == ManifestA || ManifestC == ManifestB)
-		{
-			UE_LOG(LogBuildPatchTool, Error, TEXT("ManifestC parameter can't point to the same file as ManifestA or ManifestB"));
-			return false;
-		}
-
 		// Optional list to pick specific files, otherwise it is A stomped by B
 		PARSE_SWITCH(MergeFileList);
 		NormalizeUriFile(MergeFileList);
 
-		// Parse all the secret key and ID pairs.
-		FString ParseErrorMessage;
-		if (!ParsePairs(TEXT("EncryptionSecretId="), TEXT("EncryptionSecretKey="), Switches, EncryptionSecrets, ParseErrorMessage) && !ParseErrorMessage.IsEmpty())
-		{
-			UE_LOG(LogBuildPatchTool, Error, TEXT("%s"), *ParseErrorMessage);
-			return false;
-		}
-
 		return true;
 #undef PARSE_SWITCH
 	}
 
-	void LogConfig() const
-	{
-		UE_LOG(LogBuildPatchTool, Log, TEXT("-----Configuration for MergeManifestMode------"));
-		UE_LOG(LogBuildPatchTool, Log, TEXT("   ManifestA: %s"), *ManifestA);
-		UE_LOG(LogBuildPatchTool, Log, TEXT("   ManifestB: %s"), *ManifestB);
-		UE_LOG(LogBuildPatchTool, Log, TEXT("   ManifestC: %s"), *ManifestC);
-		UE_LOG(LogBuildPatchTool, Log, TEXT("   BuildVersion: %s"), *BuildVersion);
-		UE_LOG(LogBuildPatchTool, Log, TEXT("   MergeFileList: %s"), *MergeFileList);
-		for (const TPair<FGuid, TArray<uint8>>& EncryptionSecret : EncryptionSecrets)
-		{
-			UE_LOG(LogBuildPatchTool, Log, TEXT("   AvailableSecretId: %s"), *EncryptionSecret.Key.ToString());
-		}
-	}
-
 private:
-	static TCHAR const* const MODE_NAME;
 	IBuildPatchServicesModule& BpsInterface;
-	const TCHAR* CommandLine;
 	bool bHelp;
 	FString ManifestA;
 	FString ManifestB;
 	FString ManifestC;
 	FString BuildVersion;
 	FString MergeFileList;
-	TMap<FGuid, TArray<uint8>> EncryptionSecrets;
 };
 
-IMPLEMENT_BPT_MODE(MergeManifests, FMergeManifestToolMode);
+BuildPatchTool::IToolModeRef BuildPatchTool::FMergeManifestToolModeFactory::Create(IBuildPatchServicesModule& BpsInterface)
+{
+	return MakeShareable(new FMergeManifestToolMode(BpsInterface));
+}
diff --git a/Engine/Source/Programs/BuildPatchTool/Private/ToolModes/PackageChunksMode.cpp b/Engine/Source/Programs/BuildPatchTool/Private/ToolModes/PackageChunksMode.cpp
--- a/Engine/Source/Programs/BuildPatchTool/Private/ToolModes/PackageChunksMode.cpp
+++ b/Engine/Source/Programs/BuildPatchTool/Private/ToolModes/PackageChunksMode.cpp
@@ -1,23 +1,19 @@
 // Copyright Epic Games, Inc. All Rights Reserved.
 
 #include "ToolModes/PackageChunksMode.h"
-
-#include "Algo/Find.h"
-#include "Algo/Count.h"
-#include "BuildPatchTool.h"
 #include "Interfaces/IBuildPatchServicesModule.h"
+#include "BuildPatchTool.h"
 #include "Misc/CommandLine.h"
 #include "Misc/Paths.h"
-#include "ToolModes/ToolModesHelp.h"
+#include "Algo/Find.h"
 
 using namespace BuildPatchTool;
 
 class FPackageChunksToolMode : public IToolMode
 {
 public:
-	FPackageChunksToolMode(IBuildPatchServicesModule& InBpsInterface, const TCHAR* InCommandLine)
+	FPackageChunksToolMode(IBuildPatchServicesModule& InBpsInterface)
 		: BpsInterface(InBpsInterface)
-		, CommandLine(InCommandLine)
 	{}
 
 	virtual ~FPackageChunksToolMode()
@@ -34,22 +30,48 @@ public:
 		// Print help if requested.
 		if (bHelp)
 		{
-			PrintHelp<FPackageChunksToolModeHelp>();
+			UE_LOG(LogBuildPatchTool, Display, TEXT("PACKAGE CHUNKS MODE"));
+			UE_LOG(LogBuildPatchTool, Display, TEXT("This tool mode supports packaging data required for an installation into larger files which can be used as local sources for build patch installers."));
+			UE_LOG(LogBuildPatchTool, Display, TEXT(""));
+			UE_LOG(LogBuildPatchTool, Display, TEXT("Required arguments:"));
+			UE_LOG(LogBuildPatchTool, Display, TEXT("  -mode=PackageChunks  Must be specified to launch the tool in package chunks mode."));
+			UE_LOG(LogBuildPatchTool, Display, TEXT("  -FeatureLevel=Latest Specifies the client feature level to output data for. See BuildPatchServices::EFeatureLevel for possible values."));
+			UE_LOG(LogBuildPatchTool, Display, TEXT("  -ManifestFile=\"\"     Specifies in quotes the file path to the manifest to enumerate chunks from."));
+			UE_LOG(LogBuildPatchTool, Display, TEXT("  -OutputFile=\"\"       Specifies in quotes the file path the output package. Extension of .chunkdb will be added if not present."));
+			UE_LOG(LogBuildPatchTool, Display, TEXT(""));
+			UE_LOG(LogBuildPatchTool, Display, TEXT("Optional arguments:"));
+			UE_LOG(LogBuildPatchTool, Display, TEXT("  -PrevManifestFile=\"\" Specifies in quotes the file path to a manifest for a previous build, this will be used to filter out chunks, such that the"));
+			UE_LOG(LogBuildPatchTool, Display, TEXT("                       produced chunkdb files will only contain chunks required to patch from this build to the one described by ManifestFile."));
+			UE_LOG(LogBuildPatchTool, Display, TEXT("  -CloudDir=\"\"         Specifies in quotes the cloud directory where chunks to be packaged can be found."));
+			UE_LOG(LogBuildPatchTool, Display, TEXT("  -MaxOutputFileSize=  When specified, the size of each output file (in bytes) will be limited to a maximum of the provided value."));
+			UE_LOG(LogBuildPatchTool, Display, TEXT("  -ResultDataFile=\"\"   Specifies in quotes the file path where the results will be exported as a JSON object."));
+			UE_LOG(LogBuildPatchTool, Display, TEXT("  -TagSets=\",t1,t2\"     Specifies in quotes a comma seperated tagset for filtering of data saved. Multiple sets can also be provided to split the chunkdb files by tagsets."));
+			UE_LOG(LogBuildPatchTool, Display, TEXT("                        Untagged files will be referenced with an empty tag, which you can specify using an extra comma."));
+			UE_LOG(LogBuildPatchTool, Display, TEXT("  -PrevTagSet=\",tA,tB\"  Specifies in quotes a comma seperated tagset for filtering of input data usable from PrevManifestFile. This will increase the amount of chunks"));
+			UE_LOG(LogBuildPatchTool, Display, TEXT("                        saved out by reducing the number of files from the input manifest that are assumed usable. Only one PrevTagSet should be provided."));
+			UE_LOG(LogBuildPatchTool, Display, TEXT("                        Untagged files will be referenced with an empty tag, which you can specify using an extra comma."));
+			UE_LOG(LogBuildPatchTool, Display, TEXT(""));
+			UE_LOG(LogBuildPatchTool, Display, TEXT("NB: If CloudDir is not specified, the manifest file location will be used as the cloud directory."));
+			UE_LOG(LogBuildPatchTool, Display, TEXT("NB: If an optimised delta was available, the file extension .delta.chunkdb will be used."));
+			UE_LOG(LogBuildPatchTool, Display, TEXT("NB: MaxOutputFileSize is recommended to be as large as possible. The minimum individual chunkdb filesize is equal to one chunk plus chunkdb"));
+			UE_LOG(LogBuildPatchTool, Display, TEXT("    header, and thus will not result in efficient behavior."));
+			UE_LOG(LogBuildPatchTool, Display, TEXT("NB: If MaxOutputFileSize is not specified, the one output file will be produced containing all required data."));
+			UE_LOG(LogBuildPatchTool, Display, TEXT("NB: If MaxOutputFileSize is specified, the output files will be generated as Name.part01.chunkdb, Name.part02.chunkdb etc. The part number will"));
+			UE_LOG(LogBuildPatchTool, Display, TEXT("    have the number of digits required for highest numbered part."));
+			UE_LOG(LogBuildPatchTool, Display, TEXT("NB: If MaxOutputFileSize is specified, then each part can be equal to or less than the specified size, depending on the size of the last chunk"));
+			UE_LOG(LogBuildPatchTool, Display, TEXT("    that fits."));
+			UE_LOG(LogBuildPatchTool, Display, TEXT("NB: When providing multiple -TagSets= arguments, all data from the first -TagSets= arg will be saved first, followed by any extra data needed for the second -TagSets= arg, and so on in separated chunkdb files."));
+			UE_LOG(LogBuildPatchTool, Display, TEXT("    Note that this means the chunkdb files produced for the second -TagSets= arg and later will not contain some required data for that tagset if the data already got saved out as part of a previous tagset."));
+			UE_LOG(LogBuildPatchTool, Display, TEXT("    The chunkdb files are thus additive with no dupes."));
+			UE_LOG(LogBuildPatchTool, Display, TEXT("    If it is desired that each tagset's chunkdb files contain the duplicate data, then PackageChunks should be executed once per -TagSets= arg rather than once will all -TagSets= args."));
+			UE_LOG(LogBuildPatchTool, Display, TEXT("    An empty tag must be included in one of the -TagSets= args to include untagged file data in that tagset, e.g. -TagSets=\" , t1\"."));
+			UE_LOG(LogBuildPatchTool, Display, TEXT("    Adding no -TagSets= args will include all data."));
 			return EReturnCode::OK;
 		}
 
-		LogConfig();
-
-		// Register secrets.
-		for (const TPair<FGuid, TArray<uint8>>& IdKeyPair : EncryptionSecrets)
-		{
-			BpsInterface.RegisterEncryptionSecret(IdKeyPair.Key, IdKeyPair.Value);
-		}
-
 		// Setup and run.
 		BuildPatchServices::FPackageChunksConfiguration Configuration;
-		LexFromString(Configuration.FeatureLevel, *FeatureLevel);
-		if (Configuration.FeatureLevel == BuildPatchServices::EFeatureLevel::Invalid)
+		if (!BuildPatchServices::FeatureLevelFromString(*FeatureLevel, Configuration.FeatureLevel))
 		{
 			UE_LOG(LogBuildPatchTool, Error, TEXT("Provided FeatureLevel is not recognised. Invalid arg: -FeatureLevel=%s"), *FeatureLevel);
 			return EReturnCode::ArgumentProcessingError;
@@ -62,9 +84,6 @@ public:
 		Configuration.CloudDir = CloudDir;
 		Configuration.MaxOutputFileSize = MaxOutputFileSize;
 		Configuration.ResultDataFilePath = ResultDataFile;
-		Configuration.bIgnoreManifestFileInvalidTags = bIgnoreManifestFileInvalidTags;
-		Configuration.bIgnorePrevManifestFileInvalidTags = bIgnorePrevManifestFileInvalidTags;
-		Configuration.DeltaFilenameTrailer = DeltaFilenameTrailer;
 
 		// Run the enumeration routine.
 		bool bSuccess = BpsInterface.PackageChunkData(Configuration);
@@ -76,11 +95,10 @@ private:
 	bool ProcessCommandline()
 	{
 #define HAS_SWITCH(SwitchVar) (Algo::FindByPredicate(Switches, [](const FString& Elem){ return Elem.StartsWith(TEXT(#SwitchVar "="));}) != nullptr)
-#define HAS_DUP_SWITCH(SwitchVar) (Algo::CountIf(Switches, [](const FString& Elem){ return Elem.StartsWith(TEXT(#SwitchVar "="));}) > 1)
 #define PARSE_SWITCH(SwitchVar) ParseSwitch(TEXT(#SwitchVar "="), SwitchVar, Switches)
 #define PARSE_SWITCHES(SwitchesVar) ParseSwitches(TEXT(#SwitchesVar "="), SwitchesVar, Switches)
 		TArray<FString> Tokens, Switches;
-		FCommandLine::Parse(CommandLine, Tokens, Switches);
+		FCommandLine::Parse(FCommandLine::Get(), Tokens, Switches);
 
 		bHelp = ParseOption(TEXT("help"), Switches);
 		if (bHelp)
@@ -104,19 +122,6 @@ private:
 			UE_LOG(LogBuildPatchTool, Error, TEXT("ManifestFile and OutputFile are required parameters"));
 			return false;
 		}
-
-		if (ManifestFile.IsEmpty())
-		{
-			UE_LOG(LogBuildPatchTool, Error, TEXT("ManifestFile argument can't be empty"));
-			return false;
-		}
-
-		if (OutputFile.IsEmpty())
-		{
-			UE_LOG(LogBuildPatchTool, Error, TEXT("OutputFile argument can't be empty"));
-			return false;
-		}
-
 		NormalizeUriFile(ManifestFile);
 		NormalizeUriFile(OutputFile);
 
@@ -125,22 +130,8 @@ private:
 		PARSE_SWITCH(ResultDataFile);
 		NormalizeUriFile(PrevManifestFile);
 		NormalizeUriFile(ResultDataFile);
-		TArray<FString> TagSets;
 		PARSE_SWITCHES(TagSets);
 		PARSE_SWITCH(PrevTagSet);
-		PARSE_SWITCH(DeltaFilenameTrailer);
-
-		if (HAS_DUP_SWITCH(PrevTagSet))
-		{
-			UE_LOG(LogBuildPatchTool, Error, TEXT("Only one PrevTagSet should be provided"));
-			return false;
-		}
-
-		if (ManifestFile == PrevManifestFile)
-		{
-			UE_LOG(LogBuildPatchTool, Error, TEXT("ManifestFile and PreManifestFile should not point to the same file"));
-			return false;
-		}
 
 		if (!PARSE_SWITCH(CloudDir))
 		{
@@ -157,13 +148,6 @@ private:
 				UE_LOG(LogBuildPatchTool, Error, TEXT("MaxOutputFileSize must be a valid uint64"));
 				return false;
 			}
-
-			if (MaxOutputFileSize < 10000000)
-			{
-				// MaxOutputFileSize must be at least 10MB, otherwise fail.
-				UE_LOG(LogBuildPatchTool, Error, TEXT("MaxOutputFileSize must be at least 10000000 (10MB)"));
-				return false;
-			}
 		}
 		else
 		{
@@ -192,72 +176,28 @@ private:
 			}
 		}
 
-		bIgnoreManifestFileInvalidTags = ParseOption(TEXT("IgnoreManifestFileInvalidTags"), Switches);
-		bIgnorePrevManifestFileInvalidTags = ParseOption(TEXT("IgnorePrevManifestFileInvalidTags"), Switches);
-
-		// Parse all the secret key and ID pairs.
-		FString ParseErrorMessage;
-		if (!ParsePairs(TEXT("EncryptionSecretId="), TEXT("EncryptionSecretKey="), Switches, EncryptionSecrets, ParseErrorMessage) && !ParseErrorMessage.IsEmpty())
-		{
-			UE_LOG(LogBuildPatchTool, Error, TEXT("%s"), *ParseErrorMessage);
-			return false;
-		}
-
 		return true;
 #undef PARSE_SWITCHES
 #undef PARSE_SWITCH
-#undef HAS_DUP_SWITCH
 #undef HAS_SWITCH
 	}
 
-	void LogConfig() const
-	{
-		UE_LOG(LogBuildPatchTool, Log, TEXT("-----Configuration for PackageChunksMode------"));
-		UE_LOG(LogBuildPatchTool, Log, TEXT("   ManifestFile: %s"), *ManifestFile);
-		UE_LOG(LogBuildPatchTool, Log, TEXT("   PrevManifestFile: %s"), *PrevManifestFile);
-		UE_LOG(LogBuildPatchTool, Log, TEXT("   OutputFile: %s"), *OutputFile);
-		UE_LOG(LogBuildPatchTool, Log, TEXT("   ResultDataFile: %s"), *ResultDataFile);
-		UE_LOG(LogBuildPatchTool, Log, TEXT("   CloudDir: %s"), *CloudDir);
-		UE_LOG(LogBuildPatchTool, Log, TEXT("   MaxOutputFileSize: %llu"), MaxOutputFileSize);
-		for (const TSet<FString>& TagSet : TagSetArray)
-		{
-			FString TagSetString = FString::Join(TagSet.Array(), TEXT(","));
-			UE_LOG(LogBuildPatchTool, Log, TEXT("   With Tag set: %s"), *TagSetString);
-		}
-		FString PrevTagSetString = FString::Join(PrevTagSet.Array(), TEXT(","));
-		UE_LOG(LogBuildPatchTool, Log, TEXT("   Previous Tag set: %s"), *PrevTagSetString);
-		if (bIgnoreManifestFileInvalidTags)
-		{
-			UE_LOG(LogBuildPatchTool, Log, TEXT("   IgnoreManifestFileInvalidTags: Enabled"));
-		}
-		if (bIgnorePrevManifestFileInvalidTags)
-		{
-			UE_LOG(LogBuildPatchTool, Log, TEXT("   IgnorePrevManifestFileInvalidTags: Enabled"));
-		}
-		for (const TPair<FGuid, TArray<uint8>>& EncryptionSecret : EncryptionSecrets)
-		{
-			UE_LOG(LogBuildPatchTool, Log, TEXT("   AvailableSecretId: %s"), *EncryptionSecret.Key.ToString());
-		}
-	}
-
 private:
-	static TCHAR const* const MODE_NAME;
 	IBuildPatchServicesModule& BpsInterface;
-	const TCHAR* CommandLine;
 	bool bHelp;
 	FString FeatureLevel;
 	FString ManifestFile;
 	FString PrevManifestFile;
 	FString OutputFile;
 	FString ResultDataFile;
-	FString DeltaFilenameTrailer;
 	FString CloudDir;
 	uint64 MaxOutputFileSize;
+	TArray<FString> TagSets;
 	TArray<TSet<FString>> TagSetArray;
 	TSet<FString> PrevTagSet;
-	bool bIgnoreManifestFileInvalidTags;
-	bool bIgnorePrevManifestFileInvalidTags;
-	TMap<FGuid, TArray<uint8>> EncryptionSecrets;
 };
 
-IMPLEMENT_BPT_MODE(PackageChunks, FPackageChunksToolMode);
+BuildPatchTool::IToolModeRef BuildPatchTool::FPackageChunksToolModeFactory::Create(IBuildPatchServicesModule& BpsInterface)
+{
+	return MakeShareable(new FPackageChunksToolMode(BpsInterface));
+}
