BuildPatchTool non-unity compile fixes

#Synced-CL 51092368

diff --git a/Engine/Source/Programs/BuildPatchTool/Private/Tests/Mock/ToolMode.mock.h b/Engine/Source/Programs/BuildPatchTool/Private/Tests/Mock/ToolMode.mock.h
--- a/Engine/Source/Programs/BuildPatchTool/Private/Tests/Mock/ToolMode.mock.h
+++ b/Engine/Source/Programs/BuildPatchTool/Private/Tests/Mock/ToolMode.mock.h
@@ -12,6 +12,9 @@ namespace BuildPatchTool
 
 		virtual ~FMockToolMode() {}
 
-		MOCK_METHOD0(Execute, EReturnCode());
+		virtual EReturnCode Execute() override
+		{
+			return EReturnCode::OK;
+		}
 	};
 }
diff --git a/Engine/Source/Programs/BuildPatchTool/Private/ToolModes/ChunkDeltaOptimiseMode.cpp b/Engine/Source/Programs/BuildPatchTool/Private/ToolModes/ChunkDeltaOptimiseMode.cpp
--- a/Engine/Source/Programs/BuildPatchTool/Private/ToolModes/ChunkDeltaOptimiseMode.cpp
+++ b/Engine/Source/Programs/BuildPatchTool/Private/ToolModes/ChunkDeltaOptimiseMode.cpp
@@ -6,6 +6,7 @@
 #include "Common/DiffAbortThresholdArgument.h"
 #include "Interfaces/ToolMode.h"
 #include "Misc/CommandLine.h"
+#include "Misc/ConfigCacheIni.h"
 #include "Misc/Optional.h"
 #include "Misc/Paths.h"
 #include "ToolModes/ToolModesHelp.h"
