BuildPatchTool: Remove reference to unused PlatformCryptoOpenSSL

#Synced-CL 51092368

diff --git a/Engine/Source/Programs/BuildPatchTool/BuildPatchTool.Build.cs b/Engine/Source/Programs/BuildPatchTool/BuildPatchTool.Build.cs
--- a/Engine/Source/Programs/BuildPatchTool/BuildPatchTool.Build.cs
+++ b/Engine/Source/Programs/BuildPatchTool/BuildPatchTool.Build.cs
@@ -18,7 +18,6 @@ public class BuildPatchTool : ModuleRules
 				"Projects",
 				"Networking",
 				"HTTP",
-				"PlatformCryptoOpenSSL",
 			}
 		);
 		
