[BuildPatchTool] Compilation fixes for Mac.

diff --git a/Engine/Source/Programs/BuildPatchTool/Private/BuildPatchTool.cpp b/Engine/Source/Programs/BuildPatchTool/Private/BuildPatchTool.cpp
--- a/Engine/Source/Programs/BuildPatchTool/Private/BuildPatchTool.cpp
+++ b/Engine/Source/Programs/BuildPatchTool/Private/BuildPatchTool.cpp
@@ -7,6 +7,11 @@
 #include "Misc/App.h"
 #include "Misc/OutputDeviceError.h"
 #include "Http.h"
+
+#if PLATFORM_UNIX
+#include <unistd.h>
+#endif
+
 #include "UObject/UObjectBase.h"
 
 #include "BuildPatchToolVersion.h"
@@ -48,10 +53,10 @@ namespace
 		verifyf(::SetCurrentDirectoryW(FWindowsPlatformProcess::BaseDir()), TEXT("Failed to set the working directory to '%s' (%s)"),
 			FWindowsPlatformProcess::BaseDir(),
 			FWindowsPlatformMisc::GetSystemErrorMessage(SystemError, UE_ARRAY_COUNT(SystemError), 0));
-#elif PLATFORM_MAC
-		chdir([FString(FMacPlatformProcess::BaseDir()).GetNSString() fileSystemRepresentation]);
 #elif PLATFORM_UNIX
-		chdir(TCHAR_TO_ANSI(FPlatformProcess::BaseDir()));
+		// TODO why not just call FPlatformMisc::SetCurrentWorkingDirectory(*FPlatformProcess::BaseDir());
+		// which would remove all this platform specific code.  Should be no need for this, or? -jwf
+		chdir(TCHAR_TO_UTF8(FPlatformProcess::BaseDir()));
 #else
 #error Please implement SetCurrentWorkingDirectoryToBaseDir() for current platform
 #endif
