[BuildPatchTool] More Mac complation fixes, apparently PLATFORM_UNIX is not set true when PLATFORM_MAC is true.

diff --git a/Engine/Source/Programs/BuildPatchTool/Private/BuildPatchTool.cpp b/Engine/Source/Programs/BuildPatchTool/Private/BuildPatchTool.cpp
--- a/Engine/Source/Programs/BuildPatchTool/Private/BuildPatchTool.cpp
+++ b/Engine/Source/Programs/BuildPatchTool/Private/BuildPatchTool.cpp
@@ -8,7 +8,7 @@
 #include "Misc/OutputDeviceError.h"
 #include "Http.h"
 
-#if PLATFORM_UNIX
+#if PLATFORM_UNIX || PLATFORM_MAC
 #include <unistd.h>
 #endif
 
@@ -53,10 +53,10 @@ namespace
 		verifyf(::SetCurrentDirectoryW(FWindowsPlatformProcess::BaseDir()), TEXT("Failed to set the working directory to '%s' (%s)"),
 			FWindowsPlatformProcess::BaseDir(),
 			FWindowsPlatformMisc::GetSystemErrorMessage(SystemError, UE_ARRAY_COUNT(SystemError), 0));
-#elif PLATFORM_UNIX
+#elif PLATFORM_UNIX || PLATFORM_MAC
 		// TODO why not just call FPlatformMisc::SetCurrentWorkingDirectory(*FPlatformProcess::BaseDir());
 		// which would remove all this platform specific code.  Should be no need for this, or? -jwf
-		chdir(TCHAR_TO_UTF8(FPlatformProcess::BaseDir()));
+		chdir(TCHAR_TO_UTF8(*FPlatformProcess::BaseDir()));
 #else
 #error Please implement SetCurrentWorkingDirectoryToBaseDir() for current platform
 #endif
