[BuildPatchTool] one more fix for Mac/Linux

diff --git a/Engine/Source/Programs/BuildPatchTool/Private/BuildPatchTool.cpp b/Engine/Source/Programs/BuildPatchTool/Private/BuildPatchTool.cpp
--- a/Engine/Source/Programs/BuildPatchTool/Private/BuildPatchTool.cpp
+++ b/Engine/Source/Programs/BuildPatchTool/Private/BuildPatchTool.cpp
@@ -56,7 +56,8 @@ namespace
 #elif PLATFORM_UNIX || PLATFORM_MAC
 		// TODO why not just call FPlatformMisc::SetCurrentWorkingDirectory(*FPlatformProcess::BaseDir());
 		// which would remove all this platform specific code.  Should be no need for this, or? -jwf
-		chdir(TCHAR_TO_UTF8(*FPlatformProcess::BaseDir()));
+		FTCHARToUTF8 Converted(*FPlatformProcess::BaseDir());
+		chdir(Converted.Get());
 #else
 #error Please implement SetCurrentWorkingDirectoryToBaseDir() for current platform
 #endif
