[BuildPatchTool] Fix CIS issue. There was an incorrect dereference operator in a method that already returns a path a TCharPtr

#Synced-CL 51141141

diff --git a/Engine/Source/Programs/BuildPatchTool/Private/BuildPatchTool.cpp b/Engine/Source/Programs/BuildPatchTool/Private/BuildPatchTool.cpp
--- a/Engine/Source/Programs/BuildPatchTool/Private/BuildPatchTool.cpp
+++ b/Engine/Source/Programs/BuildPatchTool/Private/BuildPatchTool.cpp
@@ -56,7 +56,7 @@ namespace
 #elif PLATFORM_UNIX || PLATFORM_MAC
 		// TODO why not just call FPlatformMisc::SetCurrentWorkingDirectory(*FPlatformProcess::BaseDir());
 		// which would remove all this platform specific code.  Should be no need for this, or? -jwf
-		FTCHARToUTF8 Converted(*FPlatformProcess::BaseDir());
+		FTCHARToUTF8 Converted(FPlatformProcess::BaseDir());
 		chdir(Converted.Get());
 #else
 #error Please implement SetCurrentWorkingDirectoryToBaseDir() for current platform
