p4 integrate //UE5/Dev-Distro-5.5/Engine/Source/Runtime/Online/BuildPatchServices/...[REDACTED_HANDLE] //UE5/Main/Engine/Source/Runtime/Online/BuildPatchServices/...
p4 integrate //UE5/Dev-Distro-5.5/Engine/Source/Programs/BuildPatchTool/...[REDACTED_HANDLE] //UE5/Main/Engine/Source/Programs/BuildPatchTool/...
p4 integrate //UE5/Dev-Distro-5.5/Engine/Programs/BuildPatchTool/...[REDACTED_HANDLE] //UE5/Main/Engine/Programs/BuildPatchTool/...

#Synced-CL 51306665

diff --git a/Engine/Source/Programs/BuildPatchTool/Private/BuildPatchTool.cpp b/Engine/Source/Programs/BuildPatchTool/Private/BuildPatchTool.cpp
--- a/Engine/Source/Programs/BuildPatchTool/Private/BuildPatchTool.cpp
+++ b/Engine/Source/Programs/BuildPatchTool/Private/BuildPatchTool.cpp
@@ -353,14 +353,6 @@ EReturnCode BuildPatchToolMain(int32 ArgC, TCHAR* ArgV[])
 		UE_LOG(LogBuildPatchTool, Error, TEXT("Tool exited with %s (%d)"), LexToString(ReturnCode), (int32)ReturnCode);
 		UE_LOG(LogBuildPatchTool, Error, TEXT("For more details, see log file at %s"), *FPaths::ConvertRelativePathToFull(FPlatformOutputDevices::GetAbsoluteLogFilename()));
 	}
-
-	// Shutdown
-	RequestEngineExit(TEXT("BuildPatchToolMain Exiting"));
-
-	FEngineLoop::AppPreExit();
-	FModuleManager::Get().UnloadModulesAtShutdown();
-	FEngineLoop::AppExit();
-
 	return ReturnCode;
 }
 
diff --git a/Engine/Source/Programs/BuildPatchTool/Private/ToolModes/ChunkDeltaOptimiseMode.cpp b/Engine/Source/Programs/BuildPatchTool/Private/ToolModes/ChunkDeltaOptimiseMode.cpp
--- a/Engine/Source/Programs/BuildPatchTool/Private/ToolModes/ChunkDeltaOptimiseMode.cpp
+++ b/Engine/Source/Programs/BuildPatchTool/Private/ToolModes/ChunkDeltaOptimiseMode.cpp
@@ -65,6 +65,8 @@ public:
 		Configuration.SourceDetailsLogFilename = SourceDetailsLogFilename;
 		Configuration.ManifestBTagAliases = MoveTemp(ManifestBTagAliases);
 		Configuration.ManifestATagAliases = MoveTemp(ManifestATagAliases);
+		Configuration.ManifestAIgnoreTags = MoveTemp(ManifestAIgnoreTags);
+		Configuration.ManifestBIgnoreTags = MoveTemp(ManifestBIgnoreTags);
 
 		// Run the build generation.
 		IChunkDeltaOptimiserRef ChunkDeltaOptimiser = BpsInterface.CreateChunkDeltaOptimiser(Configuration);
@@ -98,6 +100,7 @@ private:
 					UE_LOG(LogBuildPatchTool, Error, TEXT("%s must contain a comma separated list of 'from' tags (%s)"), SwitchName, *ManyToOneString);
 					return false;
 				}
+				To.TrimStartAndEndInline();
 				if (To.Len() == 0)
 				{
 					UE_LOG(LogBuildPatchTool, Error, TEXT("%s must contain a 'to' tag after a ';' (%s)"), SwitchName, *ManyToOneString);
@@ -106,7 +109,45 @@ private:
 
 				for (const FString& FromEntry : FromArray)
 				{
-					OutMap.Add(FromEntry, To);
+					FString Trimmed = FromEntry.TrimStartAndEnd();
+					if (Trimmed.Len() == 0)
+					{
+						UE_LOG(LogBuildPatchTool, Error, TEXT("%s contained an empty entry in the 'from' list (%s)"), SwitchName, *ManyToOneString);
+						return false;
+					}
+					OutMap.Add(MoveTemp(Trimmed), To);
+				}
+			}
+		}
+		return true;
+	}
+
+	bool ParseCommaSetSwitch(const TCHAR* SwitchName, const TArray<FString>& Switches, TSet<FString>& OutValues)
+	{
+		TArray<FString> RawStrings;
+		FString SwitchNameWithEq = FString(SwitchName) + TEXT("=");
+		if (ParseSwitches(*SwitchNameWithEq, RawStrings, Switches))
+		{
+			for (const FString& ListString : RawStrings)
+			{
+				TArray<FString> Entries;
+				ListString.ParseIntoArray(Entries, TEXT(","));
+				if (Entries.Num() == 0)
+				{
+					UE_LOG(LogBuildPatchTool, Error, TEXT("%s must contain a comma separated list tags (%s)"), SwitchName, *ListString);
+					return false;
+				}
+
+				for (const FString& Entry : Entries)
+				{
+					FString Trimmed = Entry.TrimStartAndEnd();
+					if (Trimmed.Len() == 0)
+					{
+						UE_LOG(LogBuildPatchTool, Error, TEXT("%s contained an empty entry (%s)"), SwitchName, *ListString);
+						return false;
+					}
+
+					OutValues.Emplace(MoveTemp(Trimmed));
 				}
 			}
 		}
@@ -227,6 +268,22 @@ private:
 			return false; // already logged.
 		}
 
+		// Parse tags that should not affect whether a file is allowed to produce matches or not. For example,
+		// without this:
+		// filename_1.txt tag:one tag:two
+		// filename_2.txt tag:one
+		// The files will not be able to provide matches during patching to each other as their tagsets do not match.
+		// However if you pass "two" as an ignore tag, it will be removed from consideration and the matcher will only
+		// se tag:one for filename_1.txt. You can specify different ignore tags for the two manifests used for optimization.
+		if (!ParseCommaSetSwitch(TEXT("ManifestAIgnoreTags"), Switches, ManifestAIgnoreTags))
+		{
+			return false; // already logged.
+		}
+		if (!ParseCommaSetSwitch(TEXT("ManifestBIgnoreTags"), Switches, ManifestBIgnoreTags))
+		{
+			return false; // already logged.
+		}
+
 		return true;
 #undef SENTENCE_END
 #undef PARSE_UNEXPECTED_SWITCH
@@ -287,6 +344,14 @@ private:
 		{
 			UE_LOG(LogBuildPatchTool, Log, TEXT("   SourceDetailsLogFilename: %s"), *SourceDetailsLogFilename);
 		}
+		if (ManifestAIgnoreTags.Num())
+		{
+			UE_LOG(LogBuildPatchTool, Log, TEXT("   ManifestAIgnoreTags: %s"), *FString::Join(ManifestAIgnoreTags, TEXT(",")));
+		}
+		if (ManifestBIgnoreTags.Num())
+		{
+			UE_LOG(LogBuildPatchTool, Log, TEXT("   ManifestBIgnoreTags: %s"), *FString::Join(ManifestBIgnoreTags, TEXT(",")));
+		}
 	}
 
 private:
@@ -304,6 +369,8 @@ private:
 	TOptional<BuildPatchServices::FDiffAbortThreshold> DiffAbortThreshold;
 	TMap<FString, FString> ManifestBTagAliases;
 	TMap<FString, FString> ManifestATagAliases;
+	TSet<FString> ManifestAIgnoreTags;
+	TSet<FString> ManifestBIgnoreTags;
 	FString DeltaFilenameTrailer;
 	FString SourceDetailsLogFilename;
 	TMap<FGuid, TArray<uint8>> EncryptionSecrets;
diff --git a/Engine/Source/Programs/BuildPatchTool/Private/ToolModes/DiffManifestMode.cpp b/Engine/Source/Programs/BuildPatchTool/Private/ToolModes/DiffManifestMode.cpp
--- a/Engine/Source/Programs/BuildPatchTool/Private/ToolModes/DiffManifestMode.cpp
+++ b/Engine/Source/Programs/BuildPatchTool/Private/ToolModes/DiffManifestMode.cpp
@@ -205,6 +205,8 @@ private:
 		UE_LOG(LogBuildPatchTool, Log, TEXT("   InstallTagsA: %s"), *FString::Join(InstallTagsA.Array(), TEXT(",")));
 		UE_LOG(LogBuildPatchTool, Log, TEXT("   InstallTagsB: %s"), *FString::Join(InstallTagsB.Array(), TEXT(",")));
 		UE_LOG(LogBuildPatchTool, Log, TEXT("   OutputFile: %s"), *OutputFile);
+		UE_LOG(LogBuildPatchTool, Log, TEXT("   OutputPatchDescriptorPath: %s"), *PatchDescriptorPath);
+		UE_LOG(LogBuildPatchTool, Log, TEXT("   OnlyPatchDescriptors: %s"), *LexToString(bOnlyPatchDescriptors));
 		UE_LOG(LogBuildPatchTool, Log, TEXT("   RequireOptimizedDelta: %s"), *LexToString(bRequireOptimizedDelta));
 		UE_LOG(LogBuildPatchTool, Log, TEXT("   SimulateInstallTimes: %s"), *LexToString(bSimulateInstallTimes));
 		if (DeltaFilenameTrailer.Len())
diff --git a/Engine/Source/Programs/BuildPatchTool/Private/ToolModes/InstallManifestMode.cpp b/Engine/Source/Programs/BuildPatchTool/Private/ToolModes/InstallManifestMode.cpp
--- a/Engine/Source/Programs/BuildPatchTool/Private/ToolModes/InstallManifestMode.cpp
+++ b/Engine/Source/Programs/BuildPatchTool/Private/ToolModes/InstallManifestMode.cpp
@@ -11,6 +11,7 @@
 #include "Interfaces/ToolMode.h"
 #include "Interfaces/IBuildPatchServicesModule.h"
 #include "Misc/CommandLine.h"
+#include "Misc/Optional.h"
 #include "Misc/OutputDeviceRedirector.h"
 #include "Misc/Paths.h"
 #include "ToolModes/ToolModesHelp.h"
@@ -24,8 +25,8 @@ public:
 private:
 	void RequestManifestFile(const FString& ManifestUri);
 	void AppManifestHttpRequestComplete(FHttpRequestPtr HttpRequest, FHttpResponsePtr HttpResponse, bool bSucceeded, FString ManifestUri);
-	void FilterInstallation(const IBuildManifestRef& BuildManifest);
-	void StartInstallation(const IBuildManifestRef& BuildManifest);
+	void FilterInstallation();
+	void StartInstallation();
 	void OnInstallComplete(const IBuildInstallerRef& Installer);
 	void ExitTool(BuildPatchTool::EReturnCode InReturnCode);
 	bool ProcessCommandline();
@@ -37,14 +38,12 @@ private:
 	const TCHAR* CommandLine;
 	BuildPatchTool::EReturnCode ReturnCode;
 	bool bHelp;
-	FString Manifest;
-	FString PrevManifest;
+	TArray<FInstallActionArgs> InstallActions;
 	FString OutputDir;
 	TSet<FString> FileList;
 	TArray<FString> CloudDirs;
 	TMap<FGuid, TArray<uint8>> EncryptionSecrets;
-	IBuildManifestPtr InstallManifest;
-	IBuildManifestPtr CurrentManifest;
+	TMap<FString, IBuildManifestPtr> ActionManifests;
 	IBuildInstallerPtr BuildInstaller;
 	bool bRejectSymlinks = false;
 };
@@ -90,12 +89,24 @@ BuildPatchTool::EReturnCode FInstallManifestToolMode::Execute()
 	const float MainsFramerate = 100.0f;
 	const float MainsFrameTime = 1.0f / MainsFramerate;
 
-	// Kick off the first stage then run the tick loop.
-	if (!PrevManifest.IsEmpty())
+	// Setup and kick off all relevant manifest loads.
+	for (const FInstallActionArgs& InstallAction : InstallActions)
 	{
-		RequestManifestFile(PrevManifest);
+		if (InstallAction.PrevManifestUri.IsSet())
+		{
+			ActionManifests.Add(InstallAction.PrevManifestUri.GetValue(), IBuildManifestPtr());
+		}
+		if (InstallAction.NewManifestUri.IsSet())
+		{
+			ActionManifests.Add(InstallAction.NewManifestUri.GetValue(), IBuildManifestPtr());
+		}
 	}
-	RequestManifestFile(Manifest);
+	for (TPair<FString, IBuildManifestPtr>& ActionManifest : ActionManifests)
+	{
+		RequestManifestFile(ActionManifest.Key);
+	}
+
+	// Run the tick loop.
 	while (!IsEngineExitRequested())
 	{
 		// Increment global frame counter once for each app tick.
@@ -139,26 +150,24 @@ void FInstallManifestToolMode::RequestManifestFile(const FString& ManifestUri)
 	else
 	{
 		// Read direct from file.
-		IBuildManifestPtr LoadedManifest = BpsInterface.LoadManifestFromFile(ManifestUri);
-		if (!LoadedManifest.IsValid())
-		{
-			UE_LOG(LogBuildPatchTool, Error, TEXT("Unable to read manifest from %s"), *ManifestUri);
-			ExitTool(BuildPatchTool::EReturnCode::InvalidData);
-		}
-		else if(ManifestUri == Manifest)
+		IBuildManifestPtr* ManifestToLoad = ActionManifests.Find(ManifestUri);
+		if (ManifestToLoad == nullptr)
 		{
-			InstallManifest = LoadedManifest;
+			UE_LOG(LogBuildPatchTool, Error, TEXT("Missing ActionManifest request %s"), *ManifestUri);
+			ExitTool(BuildPatchTool::EReturnCode::ToolFailure);
 		}
 		else
 		{
-			CurrentManifest = LoadedManifest;
-		}
-
-		// Are we done loading?
-		const bool bLoadComplete = (PrevManifest.IsEmpty() || CurrentManifest.IsValid()) && InstallManifest.IsValid();
-		if (bLoadComplete)
-		{
-			FilterInstallation(InstallManifest.ToSharedRef());
+			*ManifestToLoad = BpsInterface.LoadManifestFromFile(ManifestUri);
+			if (!ManifestToLoad->IsValid())
+			{
+				UE_LOG(LogBuildPatchTool, Error, TEXT("Unable to read manifest from %s"), *ManifestUri);
+				ExitTool(BuildPatchTool::EReturnCode::InvalidData);
+			}
+			else
+			{
+				FilterInstallation();
+			}
 		}
 	}
 }
@@ -167,26 +176,24 @@ void FInstallManifestToolMode::AppManifestHttpRequestComplete(FHttpRequestPtr Ht
 {
 	if (bSucceeded && HttpResponse.IsValid())
 	{
-		IBuildManifestPtr LoadedManifest = BpsInterface.MakeManifestFromData(HttpResponse->GetContent());
-		if (!LoadedManifest.IsValid())
-		{
-			UE_LOG(LogBuildPatchTool, Error, TEXT("Unable to serialise manifest from %s"), *HttpRequest->GetURL());
-			ExitTool(BuildPatchTool::EReturnCode::InvalidData);
-		}
-		else if(ManifestUri == Manifest)
+		IBuildManifestPtr* ManifestToLoad = ActionManifests.Find(ManifestUri);
+		if (ManifestToLoad == nullptr)
 		{
-			InstallManifest = LoadedManifest;
+			UE_LOG(LogBuildPatchTool, Error, TEXT("Missing ActionManifest request %s"), *ManifestUri);
+			ExitTool(BuildPatchTool::EReturnCode::ToolFailure);
 		}
 		else
 		{
-			CurrentManifest = LoadedManifest;
-		}
-
-		// Are we done loading?
-		const bool bLoadComplete = (PrevManifest.IsEmpty() || CurrentManifest.IsValid()) && InstallManifest.IsValid();
-		if (bLoadComplete)
-		{
-			FilterInstallation(InstallManifest.ToSharedRef());
+			*ManifestToLoad = BpsInterface.MakeManifestFromData(HttpResponse->GetContent());
+			if (!ManifestToLoad->IsValid())
+			{
+				UE_LOG(LogBuildPatchTool, Error, TEXT("Unable to serialise manifest from %s"), *HttpRequest->GetURL());
+				ExitTool(BuildPatchTool::EReturnCode::InvalidData);
+			}
+			else
+			{
+				FilterInstallation();
+			}
 		}
 	}
 	else
@@ -196,37 +203,94 @@ void FInstallManifestToolMode::AppManifestHttpRequestComplete(FHttpRequestPtr Ht
 	}
 }
 
-void FInstallManifestToolMode::FilterInstallation(const IBuildManifestRef& BuildManifest)
+void FInstallManifestToolMode::FilterInstallation()
 {
+	// Ignore until we have loaded all manifest.
+	const bool bLoadComplete = Algo::AllOf(ActionManifests, [](const TPair<FString, IBuildManifestPtr>& Elem) { return Elem.Value.IsValid(); });
+	if (!bLoadComplete)
+	{
+		return;
+	}
+
+	// Filter all manifests by provided file list.
 	if (FileList.Num() > 0)
 	{
-		IBuildManifestPtr FilteredManifest = BpsInterface.MergeManifests(BuildManifest, BuildManifest, BuildManifest->GetVersionString(), TSet<FString>(), FileList);
-		if (FilteredManifest.IsValid())
+		TSet<FString> ResultingFileList;
+		for (TPair<FString, IBuildManifestPtr>& ActionManifest : ActionManifests)
 		{
-			StartInstallation(FilteredManifest.ToSharedRef());
+			IBuildManifestPtr& FilteredManifest = ActionManifest.Value;
+			TSet<FString> FilteredFileList = FileList.Intersect(TSet<FString>(FilteredManifest->GetBuildFileList()));
+			ResultingFileList.Append(FilteredFileList);
+			FilteredManifest = BpsInterface.MergeManifests(FilteredManifest.ToSharedRef(), FilteredManifest.ToSharedRef(), FilteredManifest->GetVersionString(), TSet<FString>(), FilteredFileList);
+			if (!FilteredManifest.IsValid())
+			{
+				UE_LOG(LogBuildPatchTool, Error, TEXT("Manifest filtering failed, check file list and encryption secrets."));
+				ExitTool(BuildPatchTool::EReturnCode::ArgumentProcessingError);
+				return;
+			}
 		}
-		else
+		// Check that all files in FileList were indeed found in one of the actions
+		// MergeManifests will not fail because we are filtering the FileList input by action  manifest in order
+		// to support both multiple actions and a FileList together
+		bool bAllFilesFound = true;
+		for (const FString& File : FileList)
 		{
-			UE_LOG(LogBuildPatchTool, Error, TEXT("Manifest filtering failed, check file list."));
+			if (!ResultingFileList.Contains(File))
+			{
+				UE_LOG(LogBuildPatchTool, Error, TEXT("Could not find file in FileList: %s"), *File);
+				bAllFilesFound = false;
+			}
+		}
+		if (!bAllFilesFound)
+		{
+			// Error is logged above.
 			ExitTool(BuildPatchTool::EReturnCode::ArgumentProcessingError);
+			return;
 		}
 	}
-	else
-	{
-		StartInstallation(BuildManifest);
-	}
+
+	StartInstallation();
 }
 
-void FInstallManifestToolMode::StartInstallation(const IBuildManifestRef& BuildManifest)
+void FInstallManifestToolMode::StartInstallation()
 {
 	using namespace BuildPatchServices;
 
+	// Enumerate secrets.
 	TSet<FGuid> AvailableSecrets;
+	TSet<FGuid> NecessaryEncryptionSecrets;
 	Algo::Transform(EncryptionSecrets, AvailableSecrets, &TPair<FGuid, TArray<uint8>>::Key);
-	TSet<FGuid> NecessaryEncryptionSecrets = BuildManifest->GetNecessaryEncryptionSecretIds();
+	for (TPair<FString, IBuildManifestPtr>& ActionManifest : ActionManifests)
+	{
+		NecessaryEncryptionSecrets.Append(ActionManifest.Value->GetNecessaryEncryptionSecretIds());
+	}
+
+	// Configure the installer actions.
+	TArray<FInstallerAction> InstallerActions;
+	for (const FInstallActionArgs& InstallAction : InstallActions)
+	{
+		// ActionManifests has been previously validated and the tool exits if we are missing entries that are required,
+		// or if any uris failed to serialise into a valid manifest.
+
+		// Uninstall if NewManifestUri was not provided.
+		if (!InstallAction.NewManifestUri.IsSet())
+		{
+			InstallerActions.Add(FInstallerAction::MakeUninstall(ActionManifests.FindRef(InstallAction.PrevManifestUri.GetValue()).ToSharedRef()));
+		}
+		// Install if PrevManifestUri was not prvidedo.
+		else if (!InstallAction.PrevManifestUri.IsSet())
+		{
+			InstallerActions.Add(FInstallerAction::MakeInstall(ActionManifests.FindRef(InstallAction.NewManifestUri.GetValue()).ToSharedRef(), InstallAction.InstallTags.GetValue()));
+		}
+		// Otherwise, it's an update.
+		else
+		{
+			InstallerActions.Add(FInstallerAction::MakeUpdate(ActionManifests.FindRef(InstallAction.PrevManifestUri.GetValue()).ToSharedRef(), ActionManifests.FindRef(InstallAction.NewManifestUri.GetValue()).ToSharedRef(), InstallAction.InstallTags.GetValue()));
+		}
+	}
 
 	// Configure the installer.
-	FBuildInstallerConfiguration Configuration({ FInstallerAction::MakeInstallOrUpdate(CurrentManifest, BuildManifest) });
+	FBuildInstallerConfiguration Configuration(MoveTemp(InstallerActions));
 	Configuration.InstallDirectory = OutputDir;
 	Configuration.StagingDirectory = FPaths::Combine(OutputDir, TEXT(".bptstage"));
 	Configuration.CloudDirectories = CloudDirs;
@@ -236,6 +300,8 @@ void FInstallManifestToolMode::StartInstallation(const IBuildManifestRef& BuildM
 	{
 		Configuration.InstallMode = EInstallMode::Preload;
 	}
+	// Minimal verification
+	Configuration.VerifyMode = EVerifyMode::FileSizeCheckTouchedFiles;
 
 	// Start the installer.
 	FBuildPatchInstallerDelegate InstallerCompleteDelegate = FBuildPatchInstallerDelegate::CreateRaw(this, &FInstallManifestToolMode::OnInstallComplete);
@@ -275,6 +341,7 @@ bool FInstallManifestToolMode::ProcessCommandline()
 
 	TArray<FString> Tokens, Switches;
 	FCommandLine::Parse(CommandLine, Tokens, Switches);
+	FString ParseErrorMessage;
 
 	bHelp = ParseOption(TEXT("help"), Switches);
 	if (bHelp)
@@ -282,16 +349,19 @@ bool FInstallManifestToolMode::ProcessCommandline()
 		return true;
 	}
 
-	PARSE_SWITCH(Manifest);
-	PARSE_SWITCH(PrevManifest);
+	// Parse required install actions
+	if (!ParseInstallActions(Switches, InstallActions, ParseErrorMessage))
+	{
+		UE_LOG(LogBuildPatchTool, Error, TEXT("%s"), *ParseErrorMessage);
+		return false;
+	}
+
 	PARSE_SWITCH(OutputDir);
 	PARSE_SWITCH(FileList);
 	PARSE_SWITCH(CloudDirs);
 
 	bRejectSymlinks = ParseOption(TEXT("RejectSymlinks"), Switches);
 
-	NormalizeUriFile(Manifest);
-	NormalizeUriFile(PrevManifest);
 	NormalizeUriPath(OutputDir);
 	for (FString& File : FileList)
 	{
@@ -301,13 +371,24 @@ bool FInstallManifestToolMode::ProcessCommandline()
 	{
 		NormalizeUriFile(CloudDir);
 	}
+
+	// If specific cloud dirs were not provided, we take the paths of the manifests in the actions we found.
 	if (CloudDirs.Num() == 0)
 	{
-		CloudDirs.Add(FPaths::GetPath(Manifest));
+		for (const FInstallActionArgs& InstallAction : InstallActions)
+		{
+			if (InstallAction.PrevManifestUri.IsSet())
+			{
+				CloudDirs.AddUnique(FPaths::GetPath(InstallAction.PrevManifestUri.GetValue()));
+			}
+			if (InstallAction.NewManifestUri.IsSet())
+			{
+				CloudDirs.AddUnique(FPaths::GetPath(InstallAction.NewManifestUri.GetValue()));
+			}
+		}
 	}
 
 	// Parse all the secret key and ID pairs.
-	FString ParseErrorMessage;
 	if (!ParsePairs(TEXT("EncryptionSecretId="), TEXT("EncryptionSecretKey="), Switches, EncryptionSecrets, ParseErrorMessage) && !ParseErrorMessage.IsEmpty())
 	{
 		UE_LOG(LogBuildPatchTool, Error, TEXT("%s"), *ParseErrorMessage);
@@ -321,8 +402,25 @@ bool FInstallManifestToolMode::ProcessCommandline()
 void FInstallManifestToolMode::LogConfig() const
 {
 	UE_LOG(LogBuildPatchTool, Log, TEXT("-----Configuration for InstallManifest------"));
-	UE_LOG(LogBuildPatchTool, Log, TEXT("  PrevManifest:   %s"), *PrevManifest);
-	UE_LOG(LogBuildPatchTool, Log, TEXT("  Manifest:       %s"), *Manifest);
+	for (const FInstallActionArgs& InstallAction : InstallActions)
+	{
+		if (InstallAction.PrevManifestUri.IsSet())
+		{
+			UE_LOG(LogBuildPatchTool, Log, TEXT("  PrevManifest:   %s"), *InstallAction.PrevManifestUri.GetValue());
+		}
+		if (InstallAction.NewManifestUri.IsSet())
+		{
+			UE_LOG(LogBuildPatchTool, Log, TEXT("  NewManifest:    %s"), *InstallAction.NewManifestUri.GetValue());
+		}
+		if (InstallAction.InstallTags.IsSet())
+		{
+			for (const FString& Tag : InstallAction.InstallTags.GetValue())
+			{
+				UE_LOG(LogBuildPatchTool, Log, TEXT("  InstallTag:     %s"), *Tag);
+			}
+		}
+	}
+
 	UE_LOG(LogBuildPatchTool, Log, TEXT("  RejectSymlinks: %s"), bRejectSymlinks ? TEXT("true") : TEXT("false"));
 	UE_LOG(LogBuildPatchTool, Log, TEXT("  OutputDir:      %s"), *OutputDir);
 	bool bFirst = true;
diff --git a/Engine/Source/Programs/BuildPatchTool/Private/ToolModes/OnlineToolMode.cpp b/Engine/Source/Programs/BuildPatchTool/Private/ToolModes/OnlineToolMode.cpp
new file mode 100644
--- /dev/null
+++ b/Engine/Source/Programs/BuildPatchTool/Private/ToolModes/OnlineToolMode.cpp
@@ -0,0 +1,1617 @@
+// Copyright Epic Games, Inc. All Rights Reserved.
+
+#if WITH_ONLINEMCP
+#include "ToolModes/OnlineToolMode.h"
+
+#include "BuildPatchOnlineMcpModule.h"
+#include "BuildPatchTool.h"
+#include "BuildPatchToolVersion.h"
+#include "Common/ManifestMetadataSerializer.h"
+#include "Common/MetadataSerialiser.h"
+#include "Common/DiffAbortThresholdArgument.h"
+#include "HAL/FileManager.h"
+#include "Misc/CommandLine.h"
+#include "Misc/Paths.h"
+#include "ToolModes/HelpToolMode.h"
+#include "ToolModes/ToolModesHelp.h"
+
+template <class THelpProvider> class FOnlineHelpModeWrapper : public BuildPatchTool::IToolMode
+{
+public:
+	FOnlineHelpModeWrapper(BuildPatchTool::EReturnCode InReturnCode)
+		: HelpProvider(new BuildPatchTool::THelpToolMode<THelpProvider>())
+		, ReturnCode(InReturnCode)
+	{
+
+	}
+
+	virtual BuildPatchTool::EReturnCode Execute() override
+	{
+		if (ReturnCode != BuildPatchTool::EReturnCode::OK)
+		{
+			UE_LOG(LogBuildPatchTool, Error, TEXT("ERROR: No mode was detected on the commandline, please check the provided args. You can use -help for information on available modes."));
+			UE_LOG(LogBuildPatchTool, Display, TEXT(""));
+		}
+		HelpProvider->Execute();
+		return ReturnCode;
+	}
+
+private:
+	TUniquePtr<BuildPatchTool::IToolMode> HelpProvider;
+	BuildPatchTool::EReturnCode ReturnCode;
+};
+
+BuildPatchTool::EReturnCode ToBuildPatchToolReturnCode(EOnlineModeError RunError)
+{
+	switch (RunError)
+	{
+		case EOnlineModeError::NoError:
+			return BuildPatchTool::EReturnCode::OK;
+			break;
+		case EOnlineModeError::OnlineMcpError:
+			return BuildPatchTool::EReturnCode::OnlineMcpError;
+			break;
+		case EOnlineModeError::ToolFailure:
+			return BuildPatchTool::EReturnCode::ToolFailure;
+			break;
+		case EOnlineModeError::DirectoryNotFound:
+			return BuildPatchTool::EReturnCode::DirectoryNotFound;
+			break;
+		case EOnlineModeError::FileNotFound:
+			return BuildPatchTool::EReturnCode::FileNotFound;
+			break;
+		case EOnlineModeError::ArgumentProcessingError:
+			return BuildPatchTool::EReturnCode::ArgumentProcessingError;
+			break;
+		case EOnlineModeError::InvalidData:
+			return BuildPatchTool::EReturnCode::InvalidData;
+			break;
+		default:
+			check(false);
+			return BuildPatchTool::EReturnCode::ToolFailure;
+			break;
+	}
+}
+
+template<class TMode> BuildPatchTool::EReturnCode ExecuteHelper(TMode& Mode, BuildPatchTool::EReturnCode OverrideReturnCode)
+{
+	if (BuildPatchTool::EReturnCode::OK != OverrideReturnCode)
+	{
+		return OverrideReturnCode;
+	}
+	EOnlineModeError RunError = Mode->Run();
+	return ToBuildPatchToolReturnCode(RunError);
+}
+
+TArray<FString> GetAllowedCdns()
+{
+	TArray<FString> AllowedCdns;
+	GConfig->GetArray(TEXT("Portal.DistributionOverrides"), TEXT("AllowedDistributionPoints"), AllowedCdns, GEngineIni);
+	return AllowedCdns;
+}
+
+class FBinaryDeltaOptimiseToolModeWrapper : public BuildPatchTool::IToolMode
+{
+public:
+	using HelpType = BuildPatchTool::FBinaryDeltaOptimizeHelp;
+
+	FBinaryDeltaOptimiseToolModeWrapper(const BuildPatchTool::FToolModeConfiguration& Configuration)
+		: CommandLine(Configuration.CommandLine)
+		, OverrideReturnCode(BuildPatchTool::EReturnCode::OK)
+	{
+		FBinaryDeltaOptimiseConfig Config;
+		if (!ProcessCommandLine(Config))
+		{
+			OverrideReturnCode = BuildPatchTool::EReturnCode::ArgumentProcessingError;
+		}
+
+		const bool bNeedToForceNonSignedUrls = Config.bForceSignedUrls.IsSet() && Config.bForceSignedUrls.GetValue() == false;
+		// In all cases where we don't need to force non-signed urls, we need to create and set SignedUriProviders
+		if (!bNeedToForceNonSignedUrls)
+		{
+			TArray<FString> AllowedCdns = GetAllowedCdns();
+			FSignedUriProviderConfig SignedUriProviderConfigA{ Config.ArtifactIdA, Config.BuildVersionA, Config.ProductId, Config.OrganizationId, AllowedCdns, false };
+			FSignedUriProviderConfig SignedUriProviderConfigB{ Config.ArtifactIdB, Config.BuildVersionB, Config.ProductId, Config.OrganizationId, AllowedCdns, false };
+
+			Config.SignedUriProviderA = Configuration.BpMcpModule.CreateSignedUriProvider(SignedUriProviderConfigA);
+			Config.SignedUriProviderB = Configuration.BpMcpModule.CreateSignedUriProvider(SignedUriProviderConfigB);
+		}
+
+		Mode = TUniquePtr<IBinaryDeltaOptimiseMode>(Configuration.BpMcpModule.CreateBinaryDeltaOptimiseMode(Configuration.BpsInterface, MoveTemp(Config)));
+	}
+public:
+	BuildPatchTool::EReturnCode Execute() override
+	{
+		return ExecuteHelper(Mode, OverrideReturnCode);
+	}
+
+	bool IsOnline() const override
+	{
+		return true;
+	}
+
+private:
+	bool ProcessCommandLine(FBinaryDeltaOptimiseConfig& Config)
+	{
+#define PARSE_SWITCH(Switch) ParseSwitch(TEXT(#Switch "="), Config.Switch, Switches)
+#define _PARSE_REQUIRED_SWITCH(Switch, ConfigSwitch, bExpect) \
+		if (ParseSwitch(TEXT(#Switch "="), ConfigSwitch, Switches) != bExpect) \
+		{ \
+			UE_LOG(LogBuildPatchTool, Error, #Switch SENTENCE_END); \
+			bParametersOk = false; \
+		}
+#define PARSE_REQUIRED_SWITCH(Switch) _PARSE_REQUIRED_SWITCH(Switch, Config.Switch, true)
+#define PARSE_REQUIRED_SWITCH_WITH_ALIAS(Switch, Alias) _PARSE_REQUIRED_SWITCH(Switch, Config.Alias, true)
+#define PARSE_UNEXPECTED_SWITCH(Switch) do{FString Switch##Temp; _PARSE_REQUIRED_SWITCH(Switch, Switch##Temp, false);}while(false)
+
+		TArray<FString> Tokens, Switches;
+		FCommandLine::Parse(CommandLine, Tokens, Switches);
+
+		// If help is requested, the help mode wrapper should be used.
+		if (ParseOption(TEXT("help"), Switches))
+		{
+			// If help is requested, the help mode wrapper should be used.
+			return false;
+		}
+
+		const bool bHasForceNonSignedUrlsFlag = ParseOption(TEXT("forcenonsignedurls"), Switches);
+		const bool bHasForceSignedUrlsFlag = ParseOption(TEXT("forcesignedurls"), Switches);
+
+		// Use only one of these arguments.
+		if (bHasForceNonSignedUrlsFlag && bHasForceSignedUrlsFlag)
+		{
+			UE_LOG(LogBuildPatchTool, Error, TEXT("The forcenonsignedurls and forcesignedurls arguments cannot be used simultaneously. Use only one of them."));
+			return false;
+		}
+
+		if (bHasForceNonSignedUrlsFlag)
+		{
+			Config.bForceSignedUrls = false;
+		}
+		else if (bHasForceSignedUrlsFlag)
+		{
+			Config.bForceSignedUrls = true;
+		}
+
+#define SENTENCE_END TEXT(" is a required parameter used when uploading to Epic services.")
+		bool bParametersOk = true;
+		// Get all required online parameters, we allow the auth manager to report errors with required params for client secret.
+		FString ArtifactId;
+		PARSE_REQUIRED_SWITCH(OrganizationId);
+		PARSE_REQUIRED_SWITCH(ProductId);
+		PARSE_REQUIRED_SWITCH(BuildVersionA);
+		PARSE_REQUIRED_SWITCH(BuildVersionB);
+		PARSE_REQUIRED_SWITCH_WITH_ALIAS(ArtifactId, ArtifactIdA);
+		PARSE_REQUIRED_SWITCH(CloudDir);
+#undef SENTENCE_END
+#define SENTENCE_END TEXT(" should not be used when uploading to Epic services.")
+		// Check for params that should not have been provided for online mode.
+		PARSE_UNEXPECTED_SWITCH(ManifestA);
+		PARSE_UNEXPECTED_SWITCH(ManifestB);
+		PARSE_UNEXPECTED_SWITCH(DownloadCloudDirectory);
+		PARSE_UNEXPECTED_SWITCH(OutputChunkSize);
+#undef SENTENCE_END
+		if (!bParametersOk)
+		{
+			return false;
+		}
+
+		if (Config.CloudDir.IsEmpty() || Config.ArtifactIdA.IsEmpty() || Config.OrganizationId.IsEmpty() || Config.ProductId.IsEmpty())
+		{
+			UE_LOG(LogBuildPatchTool, Error, TEXT("CloudDir, ArtifactId, OrganizationId, ProductId arguments can't be empty"));
+			return false;
+		}
+
+		if (Config.BuildVersionA.IsEmpty() || Config.BuildVersionA.Len() > 100)
+		{
+			UE_LOG(LogBuildPatchTool, Error, TEXT("Please ensure you've set the BuildVersionA for BinaryDeltaOptimise mode and it is no more than 100 characters."));
+			return false;
+		}
+
+		if (!BuildVersionContainsOnlyAllowedSymbols(Config.BuildVersionA))
+		{
+			UE_LOG(LogBuildPatchTool, Error, TEXT("BuildVersionA string should only contain characters from the following sets a-z, A-Z, 0-9, or .+-_"));
+			return false;
+		}
+
+		if (Config.BuildVersionB.IsEmpty() || Config.BuildVersionB.Len() > 100)
+		{
+			UE_LOG(LogBuildPatchTool, Error, TEXT("Please ensure you've set the BuildVersionB for BinaryDeltaOptimise mode and it is no more than 100 characters."));
+			return false;
+		}
+
+		if (!BuildVersionContainsOnlyAllowedSymbols(Config.BuildVersionB))
+		{
+			UE_LOG(LogBuildPatchTool, Error, TEXT("BuildVersionB string should only contain characters from the following sets a-z, A-Z, 0-9, or .+-_"));
+			return false;
+		}
+
+		NormalizeUriPath(Config.CloudDir);
+
+		// Get optional values.
+		if (PARSE_SWITCH(ArtifactIdB))
+		{
+			if (Config.ArtifactIdB.IsEmpty())
+			{
+				UE_LOG(LogBuildPatchTool, Error, TEXT("Argument ArtifactIdB found but empty. If you want to use it, set it to a non-empty value."));
+				return false;
+			}
+		}
+		else  // if we didn't find ArtifactIdB, then set it to ArtifactIdA value
+		{
+			Config.ArtifactIdB = Config.ArtifactIdA;
+		}
+
+		Config.DiffAbortThreshold = BuildPatchTool::FDiffAbortThresholdArgument::Parse(*this, Switches);
+
+		PARSE_SWITCH(ScanWindowSize);
+		const uint32 RequestedScanWindowSize = Config.ScanWindowSize;
+		Config.ScanWindowSize = FMath::Clamp<uint32>(Config.ScanWindowSize, 128, 128*1024);
+		if (RequestedScanWindowSize != Config.ScanWindowSize)
+		{
+			UE_LOG(LogBuildPatchTool, Warning, TEXT("Requested -ScanWindowSize=%u is outside of allowed range 128KiB >= n >= 128b. Please update your arg to be within range. Continuing with %u."), RequestedScanWindowSize, Config.ScanWindowSize);
+		}
+
+
+		return true;
+#undef PARSE_UNEXPECTED_SWITCH
+#undef PARSE_REQUIRED_SWITCH_WITH_ALIAS
+#undef PARSE_REQUIRED_SWITCH
+#undef _PARSE_REQUIRED_SWITCH
+#undef PARSE_SWITCH
+	}
+private:
+	const TCHAR* CommandLine;
+	TUniquePtr<IBinaryDeltaOptimiseMode> Mode;
+	BuildPatchTool::EReturnCode OverrideReturnCode;
+};
+
+class FUploadBinaryToolModeWrapper : public BuildPatchTool::IToolMode
+{
+public:
+	using HelpType = BuildPatchTool::FUploadBinaryHelp;
+
+	FUploadBinaryToolModeWrapper(const BuildPatchTool::FToolModeConfiguration& Configuration)
+		: CommandLine(Configuration.CommandLine)
+		, OverrideReturnCode(BuildPatchTool::EReturnCode::OK)
+	{
+		FUploadBinaryConfig Config;
+		if (!ProcessCommandline(Config, Configuration.BackendFeatureLevel))
+		{
+			OverrideReturnCode = BuildPatchTool::EReturnCode::ArgumentProcessingError;
+		}
+		ManifestMetadataSerializer = TUniquePtr<IManifestMetadataSerializer>(new FManifestMetadataSerializer);
+		Mode = TUniquePtr<IUploadBinaryMode>(Configuration.BpMcpModule.CreateUploadBinaryMode(Configuration.BpsInterface, MoveTemp(Config), [](float) {}, ManifestMetadataSerializer.Get()));
+	}
+
+public:
+	BuildPatchTool::EReturnCode Execute() override
+	{
+		return ExecuteHelper(Mode, OverrideReturnCode);
+	}
+
+	bool IsOnline() const override
+	{
+		return true;
+	}
+private:
+	bool ProcessCommandline(FUploadBinaryConfig& Config, BuildPatchServices::EFeatureLevel BackendFeatureLevel)
+	{
+#define PARSE_SWITCH(Switch) ParseSwitch(TEXT(#Switch "="), Config.Switch, Switches)
+		TArray<FString> Tokens, Switches;
+		FCommandLine::Parse(CommandLine, Tokens, Switches);
+
+		const bool bHelp = ParseOption(TEXT("help"), Switches);
+		if (bHelp)
+		{
+			// If help is requested, the help mode wrapper should be used.
+			return false;
+		}
+
+		// Grab the FeatureLevel. This is optional for UploadBinary mode, we default it to online value.
+		BuildPatchServices::EFeatureLevel FeatureLevelToSetNumeric = BackendFeatureLevel;
+
+		FString FeatureLevelString;
+		Config.bFeatureLevelPresent = ParseSwitch(TEXT("FeatureLevel="), FeatureLevelString, Switches);
+		if (Config.bFeatureLevelPresent)
+		{
+			FeatureLevelString.TrimStartAndEndInline();
+			LexFromString(FeatureLevelToSetNumeric, *FeatureLevelString);
+			if (BuildPatchServices::EFeatureLevel::Invalid == FeatureLevelToSetNumeric)
+			{
+				UE_LOG(LogBuildPatchTool, Error, TEXT("Invalid arg: -FeatureLevel=%s. FeatureLevel was provided, but not recognized. Please provide the FeatureLevel commandline argument which matches the existing client support."), *FeatureLevelString);
+				return false;
+			}
+		}
+
+		Config.FeatureLevel = FeatureLevelToSetNumeric;
+
+		FString OutputFilename;
+		if (ParseSwitch(TEXT("OutputFilename="), OutputFilename, Switches))
+		{
+			UE_LOG(LogBuildPatchTool, Error, TEXT("It is not permitted to specify OutputFilename when running against Epic's online services."));
+			return false;
+		}
+
+		// Get all required parameters
+		if (!(PARSE_SWITCH(CloudDir)
+			&& PARSE_SWITCH(BuildRoot)
+			&& PARSE_SWITCH(ArtifactId)
+			&& PARSE_SWITCH(BuildVersion)
+			&& PARSE_SWITCH(AppLaunch)
+			&& PARSE_SWITCH(AppArgs)))
+		{
+			UE_LOG(LogBuildPatchTool, Error, TEXT("CloudDir, BuildRoot, ArtifactId, BuildVersion, AppLaunch, and AppArgs are required parameters"));
+			return false;
+		}
+		NormalizeUriPath(Config.CloudDir);
+		NormalizeUriPath(Config.BuildRoot);
+		NormalizeUriFile(Config.AppLaunch);
+
+		if (Config.BuildVersion.IsEmpty() || Config.BuildVersion.Len() > 100)
+		{
+			UE_LOG(LogBuildPatchTool, Error, TEXT("Please ensure you've set the BuildVersion for UploadBinary mode and it is no more than 100 characters."));
+			return false;
+		}
+
+		if (!BuildVersionContainsOnlyAllowedSymbols(Config.BuildVersion))
+		{
+			UE_LOG(LogBuildPatchTool, Error, TEXT("BuildVersion string should only contain characters from the following sets a-z, A-Z, 0-9, or .+-_"));
+			return false;
+		}
+
+		// Trim directory names, because low level APIs will report trailing whitespace as fine for existence checks, but fail on actually writing data.
+		Config.CloudDir.TrimStartAndEndInline();
+		Config.BuildRoot.TrimStartAndEndInline();
+
+		if (Config.BuildRoot.IsEmpty())
+		{
+			UE_LOG(LogBuildPatchTool, Error, TEXT("No build root was provided."));
+			return false;
+		}
+
+		if (Config.CloudDir.IsEmpty())
+		{
+			UE_LOG(LogBuildPatchTool, Error, TEXT("No cloud directory was provided."));
+			return false;
+		}
+
+		// Ensure CloudDir and BuildRoot paths are both absolute paths and normalized before comparing them to one another.
+		Config.CloudDir = FPaths::ConvertRelativePathToFull(Config.CloudDir);
+		Config.BuildRoot = FPaths::ConvertRelativePathToFull(Config.BuildRoot);
+
+		// Check that the CloudDir and BuildRoot directories are not identical.
+		if (FPaths::IsSamePath(Config.CloudDir, Config.BuildRoot))
+		{
+			UE_LOG(LogBuildPatchTool, Error, TEXT("The BuildRoot and CloudDir directories must be different."));
+			return false;
+		}
+
+		// Check that the CloudDir is not a child of the BuildRoot.
+		if (Config.CloudDir.StartsWith(Config.BuildRoot) && (Config.CloudDir[Config.BuildRoot.Len()] == TEXT('/')))
+		{
+			UE_LOG(LogBuildPatchTool, Error, TEXT("The CloudDir directory must not be a child of the BuildRoot directory."));
+			return false;
+		}
+
+		// Check that the BuildRoot is not a child of the CloudDir.
+		if (Config.BuildRoot.StartsWith(Config.CloudDir) && (Config.BuildRoot[Config.CloudDir.Len()] == TEXT('/')))
+		{
+			UE_LOG(LogBuildPatchTool, Error, TEXT("The BuildRoot directory must not be a child of the CloudDir directory."));
+			return false;
+		}
+
+		// Check that BuildRoot exists.
+		if (!FPaths::DirectoryExists(Config.BuildRoot))
+		{
+			UE_LOG(LogBuildPatchTool, Error, TEXT("Specified BuildRoot directory does not exist."));
+			return false;
+		}
+
+		// Check that BuildRoot is not empty.
+		IFileManager& FM = IFileManager::Get();
+		TArray<FString> FoundFiles;
+		FM.FindFilesRecursive(FoundFiles, *Config.BuildRoot, TEXT("*"), true, false);
+
+		bool bBuildRootIsEmpty = true;
+		for (const FString& File : FoundFiles)
+		{
+			if (FM.FileSize(*File) > 0)
+			{
+				bBuildRootIsEmpty = false;
+				break;
+			}
+		}
+
+		if (bBuildRootIsEmpty)
+		{
+			UE_LOG(LogBuildPatchTool, Error, TEXT("BuildRoot must not be an empty directory or a directory tree containing only empty files/directories."));
+			return false;
+		}
+
+		// Get optional parameters
+		PARSE_SWITCH(ClientId);
+		PARSE_SWITCH(OrganizationId);
+		PARSE_SWITCH(ProductId);
+		PARSE_SWITCH(FileList);
+		PARSE_SWITCH(FileIgnoreList);
+		PARSE_SWITCH(FileAttributeList);
+		PARSE_SWITCH(PrereqPath);
+		PARSE_SWITCH(PrereqArgs);
+		PARSE_SWITCH(UninstallActionPath);
+		PARSE_SWITCH(UninstallActionArgs);
+		PARSE_SWITCH(AllowedArtifactIds);
+		PARSE_SWITCH(MetadataOutput);
+		PARSE_SWITCH(QueryEncryptionStatusTimeoutSec);
+		Config.bFollowSymlinks = ParseOption(TEXT("FollowSymlinks"), Switches);
+		Config.bSHA1Only = ParseOption(TEXT("SHA1Only"), Switches);
+		Config.bUseEncryption = ParseOption(TEXT("UseEncryption"), Switches);
+
+		// Note that -nofilecollisioncheck is left intentionally undocumented.
+		Config.bNoFileCollisionCheck = ParseOption(TEXT("nofilecollisioncheck"), Switches);
+
+		FString PreReqIdString;
+		ParseSwitch(TEXT("PrereqIds="), PreReqIdString, Switches);
+		if (!PreReqIdString.IsEmpty() && !ParsePrereqIds(PreReqIdString, Config.PrereqIds))
+		{
+			UE_LOG(LogBuildPatchTool, Error, TEXT("An error occurred processing comma-separated list from commandline -PrereqIds=%s"), *PreReqIdString);
+			return false;
+		}
+
+		NormalizeUriFile(Config.FileList);
+		NormalizeUriFile(Config.FileIgnoreList);
+		NormalizeUriFile(Config.FileAttributeList);
+		NormalizeUriFile(Config.PrereqPath);
+		NormalizeUriFile(Config.UninstallActionPath);
+		NormalizeUriFile(Config.MetadataOutput);
+
+		if (Config.OrganizationId.IsEmpty() != Config.ProductId.IsEmpty() || Config.OrganizationId.IsEmpty() != Config.ClientId.IsEmpty())
+		{
+			UE_LOG(LogBuildPatchTool, Error, TEXT("When uploading to Epic services, ClientID, OrganizationId, and ProductId must be specified."));
+			return false;
+		}
+
+		return true;
+#undef PARSE_SWITCH
+	}
+
+	bool ParsePrereqIds(const FString& PreReqIdString, TSet<FString>& OutPrereqIds)
+	{
+		if (PreReqIdString.Contains(Constants::Slash)
+			|| PreReqIdString.Contains(Constants::Backslash)
+			|| PreReqIdString.Contains(Constants::DoubleQuote))
+		{
+			return false;
+		}
+
+		TArray<FString> ParamValues;
+		PreReqIdString.ParseIntoArray(ParamValues, *Constants::Comma);
+		for (FString& ParamValue : ParamValues)
+		{
+			ParamValue.TrimStartAndEndInline();
+		}
+		OutPrereqIds.Append(ParamValues);
+		return true;
+	}
+
+private:
+	const TCHAR* CommandLine;
+	TUniquePtr<IUploadBinaryMode> Mode;
+	TUniquePtr<IManifestMetadataSerializer> ManifestMetadataSerializer;
+	BuildPatchTool::EReturnCode OverrideReturnCode;
+};
+
+class FReuploadBinaryToolModeWrapper : public BuildPatchTool::IToolMode
+{
+public:
+	using HelpType = BuildPatchTool::FReuploadBinaryHelp;
+
+	FReuploadBinaryToolModeWrapper(const BuildPatchTool::FToolModeConfiguration& Configuration)
+		: CommandLine(Configuration.CommandLine)
+		, OverrideReturnCode(BuildPatchTool::EReturnCode::OK)
+	{
+		FReuploadBinaryConfig Config;
+		if (!ProcessCommandline(Config))
+		{
+			OverrideReturnCode = BuildPatchTool::EReturnCode::ArgumentProcessingError;
+		}
+		Mode = TUniquePtr<IReuploadBinaryMode>(Configuration.BpMcpModule.CreateReuploadBinaryMode(Configuration.BpsInterface, MoveTemp(Config)));
+	}
+
+public:
+	BuildPatchTool::EReturnCode Execute() override
+	{
+		return ExecuteHelper(Mode, OverrideReturnCode);
+	}
+
+	bool IsOnline() const override
+	{
+		return true;
+	}
+private:
+	bool ProcessCommandline(FReuploadBinaryConfig& Config)
+	{
+#define PARSE_SWITCH(Switch) ParseSwitch(TEXT(#Switch "="), Config.Switch, Switches)
+		TArray<FString> Tokens, Switches;
+		FCommandLine::Parse(CommandLine, Tokens, Switches);
+
+		const bool bHelp = ParseOption(TEXT("help"), Switches);
+		if (bHelp)
+		{
+			// If help is requested, the help mode wrapper should be used.
+			return false;
+		}
+
+		// Grab the FeatureLevel. This is required param but safe to default, we can change this to a warning after first release, and then an error later, as part of a friendly roll out.
+		FString FeatureLevelString;
+		const bool bFeatureLevelPresent = ParseSwitch(TEXT("FeatureLevelOverride="), FeatureLevelString, Switches);
+		FeatureLevelString.TrimStartAndEndInline();
+		BuildPatchServices::EFeatureLevel FeatureLevelToSetNumeric = BuildPatchServices::EFeatureLevel::Invalid;
+		if (bFeatureLevelPresent)
+		{
+			LexFromString(FeatureLevelToSetNumeric, *FeatureLevelString);
+			if (BuildPatchServices::EFeatureLevel::Invalid == FeatureLevelToSetNumeric)
+			{
+				UE_LOG(LogBuildPatchTool, Error, TEXT("Invalid arg: -FeatureLevel=%s. FeatureLevel was provided, but not recognized. Please provide the FeatureLevel commandline argument which matches the existing client support."), *FeatureLevelString);
+				return false;
+			}
+		}
+		Config.FeatureLevelOverride = FeatureLevelToSetNumeric;
+
+		// Get all required parameters
+		if (!(PARSE_SWITCH(CloudDir)
+			&& PARSE_SWITCH(BuildRoot)
+			&& PARSE_SWITCH(ManifestFilePath)
+			&& PARSE_SWITCH(ClientId)
+			&& PARSE_SWITCH(OrganizationId)
+			&& PARSE_SWITCH(ProductId)))
+		{
+			UE_LOG(LogBuildPatchTool, Error, TEXT("ManifestFilePath, CloudDir, BuildRoot, ClientId, OrganizationId, and ProductId are required parameters"));
+			return false;
+		}
+		NormalizeUriPath(Config.CloudDir);
+		NormalizeUriPath(Config.BuildRoot);
+		NormalizeUriFile(Config.ManifestFilePath);
+
+		// Trim directory names, because low level APIs will report trailing whitespace as fine for existence checks, but fail on actually writing data.
+		Config.CloudDir.TrimStartAndEndInline();
+		Config.BuildRoot.TrimStartAndEndInline();
+
+		if (Config.BuildRoot.IsEmpty())
+		{
+			UE_LOG(LogBuildPatchTool, Error, TEXT("No build root was provided."));
+			return false;
+		}
+
+		if (Config.CloudDir.IsEmpty())
+		{
+			UE_LOG(LogBuildPatchTool, Error, TEXT("No cloud directory was provided."));
+			return false;
+		}
+
+		if (Config.ManifestFilePath.IsEmpty())
+		{
+			UE_LOG(LogBuildPatchTool, Error, TEXT("No manifest file was provided."));
+			return false;
+		}
+
+		// Ensure paths are both absolute paths and normalized before comparing them to one another.
+		Config.CloudDir = FPaths::ConvertRelativePathToFull(Config.CloudDir);
+		Config.BuildRoot = FPaths::ConvertRelativePathToFull(Config.BuildRoot);
+		Config.ManifestFilePath = FPaths::ConvertRelativePathToFull(Config.ManifestFilePath);
+
+		// Check that the CloudDir and BuildRoot directories are not identical.
+		if (FPaths::IsSamePath(Config.CloudDir, Config.BuildRoot))
+		{
+			UE_LOG(LogBuildPatchTool, Error, TEXT("The BuildRoot and CloudDir directories must be different."));
+			return false;
+		}
+
+		// Check that the CloudDir is not a child of the BuildRoot.
+		if (Config.CloudDir.StartsWith(Config.BuildRoot) && (Config.CloudDir[Config.BuildRoot.Len()] == TEXT('/')))
+		{
+			UE_LOG(LogBuildPatchTool, Error, TEXT("The CloudDir directory must not be a child of the BuildRoot directory."));
+			return false;
+		}
+
+		// Check that the BuildRoot is not a child of the CloudDir.
+		if (Config.BuildRoot.StartsWith(Config.CloudDir) && (Config.BuildRoot[Config.CloudDir.Len()] == TEXT('/')))
+		{
+			UE_LOG(LogBuildPatchTool, Error, TEXT("The BuildRoot directory must not be a child of the CloudDir directory."));
+			return false;
+		}
+
+		// Check that BuildRoot exists.
+		if (!FPaths::DirectoryExists(Config.BuildRoot))
+		{
+			UE_LOG(LogBuildPatchTool, Error, TEXT("Specified BuildRoot directory does not exist."));
+			return false;
+		}
+
+		// Check that ManifestFilePath exists.
+		if (!FPaths::FileExists(Config.ManifestFilePath))
+		{
+			UE_LOG(LogBuildPatchTool, Error, TEXT("Specified ManifestFilePath does not exist."));
+			return false;
+		}
+
+		return true;
+#undef PARSE_SWITCH
+	}
+
+private:
+	const TCHAR* CommandLine;
+	TUniquePtr<IReuploadBinaryMode> Mode;
+	BuildPatchTool::EReturnCode OverrideReturnCode;
+};
+
+class FDiffBinariesToolModeWrapper : public BuildPatchTool::IToolMode
+{
+public:
+	using HelpType = BuildPatchTool::FDiffBinariesHelp;
+
+	FDiffBinariesToolModeWrapper(const BuildPatchTool::FToolModeConfiguration& Configuration)
+		: CommandLine(Configuration.CommandLine)
+		, OverrideReturnCode(BuildPatchTool::EReturnCode::OK)
+	{
+		FDiffBinariesConfig Config;
+		if (!ProcessCommandLine(Config))
+		{
+			OverrideReturnCode = BuildPatchTool::EReturnCode::ArgumentProcessingError;
+		}
+
+		const bool bNeedToForceNonSignedUrls = Config.bForceSignedUrls.IsSet() && Config.bForceSignedUrls.GetValue() == false;
+		// In all cases where we don't need to force non-signed urls, we need to create and set SignedUriProviders
+		if (!bNeedToForceNonSignedUrls)
+		{
+			FSignedUriProviderConfig SignedUriProviderConfig{ Config.ArtifactId, Config.BuildVersionB, Config.ProductId, Config.OrganizationId, TArray<FString>() /* Disallowed CDNs are not used in this mode */, false };
+			Config.SignedUriProvider = Configuration.BpMcpModule.CreateSignedUriProvider(SignedUriProviderConfig);
+		}
+
+		Mode = TUniquePtr<IDiffBinariesMode>(Configuration.BpMcpModule.CreateDiffBinariesMode(Configuration.BpsInterface, MoveTemp(Config)));
+	}
+
+public:
+	BuildPatchTool::EReturnCode Execute() override
+	{
+		return ExecuteHelper(Mode, OverrideReturnCode);
+	}
+	bool IsOnline() const override
+	{
+		return true;
+	}
+private:
+	bool ProcessCommandLine(FDiffBinariesConfig& Config)
+	{
+#define PARSE_SWITCH(Switch) ParseSwitch(TEXT(#Switch "="), Config.Switch, Switches)
+#define PARSE_SWITCH_NO_CONFIG(Switch) ParseSwitch(TEXT(#Switch "="), Switch, Switches)
+#define _PARSE_REQUIRED_SWITCH(Switch, bExpect) \
+		if (PARSE_SWITCH(Switch) != bExpect) \
+		{ \
+			UE_LOG(LogBuildPatchTool, Error, #Switch SENTENCE_END); \
+			bParametersOk = false; \
+		}
+#define PARSE_REQUIRED_SWITCH(Switch) _PARSE_REQUIRED_SWITCH(Switch, true)
+#define PARSE_UNEXPECTED_SWITCH(Switch) \
+		if (PARSE_SWITCH_NO_CONFIG(Switch) != false) \
+		{ \
+			UE_LOG(LogBuildPatchTool, Error, #Switch SENTENCE_END); \
+			bParametersOk = false; \
+		}
+#define EXPECT_NONEMPTY_CONFIG(Switch) \
+		if (Config.Switch.IsEmpty()) \
+		{ \
+			UE_LOG(LogBuildPatchTool, Error, #Switch SENTENCE_END); \
+			bParametersOk = false; \
+		}
+
+		TArray<FString> Tokens, Switches;
+		FCommandLine::Parse(CommandLine, Tokens, Switches);
+
+		const bool bHelp = ParseOption(TEXT("help"), Switches);
+		if (bHelp)
+		{
+			// If help is requested, the help mode wrapper should be used.
+			return false;
+		}
+
+		const bool bHasForceNonSignedUrlsFlag = ParseOption(TEXT("forcenonsignedurls"), Switches);
+		const bool bHasForceSignedUrlsFlag = ParseOption(TEXT("forcesignedurls"), Switches);
+
+		// Use only one of these arguments.
+		if (bHasForceNonSignedUrlsFlag && bHasForceSignedUrlsFlag)
+		{
+			UE_LOG(LogBuildPatchTool, Error, TEXT("The forcenonsignedurls and forcesignedurls arguments cannot be used simultaneously. Use only one of them."));
+			return false;
+		}
+
+		if (bHasForceNonSignedUrlsFlag)
+		{
+			Config.bForceSignedUrls = false;
+		}
+		else if (bHasForceSignedUrlsFlag)
+		{
+			Config.bForceSignedUrls = true;
+		}
+
+#define SENTENCE_END TEXT(" is a required parameter used when uploading to Epic services.")
+		bool bParametersOk = true;
+		// Get all required online parameters, we allow the auth manager to report errors with required params for client secret.
+		PARSE_REQUIRED_SWITCH(OrganizationId);
+		PARSE_REQUIRED_SWITCH(ProductId);
+		PARSE_REQUIRED_SWITCH(BuildVersionA);
+		PARSE_REQUIRED_SWITCH(BuildVersionB);
+		PARSE_REQUIRED_SWITCH(ArtifactId);
+#undef SENTENCE_END
+#define SENTENCE_END TEXT(" should not be used when uploading to Epic services.")
+		// Check for params that should not have been provided for online mode.
+		FString ManifestA, ManifestB;
+		PARSE_UNEXPECTED_SWITCH(ManifestA);
+		PARSE_UNEXPECTED_SWITCH(ManifestB);
+#undef SENTENCE_END
+#define SENTENCE_END TEXT(" must not be empty.")
+		EXPECT_NONEMPTY_CONFIG(ArtifactId);
+		EXPECT_NONEMPTY_CONFIG(OrganizationId);
+		EXPECT_NONEMPTY_CONFIG(ProductId);
+		if (!bParametersOk)
+		{
+			return false;
+		}
+
+		if (Config.BuildVersionA.IsEmpty() || Config.BuildVersionA.Len() > 100)
+		{
+			UE_LOG(LogBuildPatchTool, Error, TEXT("Please ensure you've set the BuildVersionA for DiffBinaries mode and it is no more than 100 characters."));
+			return false;
+		}
+
+		if (!BuildVersionContainsOnlyAllowedSymbols(Config.BuildVersionA))
+		{
+			UE_LOG(LogBuildPatchTool, Error, TEXT("BuildVersionA string should only contain characters from the following sets a-z, A-Z, 0-9, or .+-_"));
+			return false;
+		}
+
+		if (Config.BuildVersionB.IsEmpty() || Config.BuildVersionB.Len() > 100)
+		{
+			UE_LOG(LogBuildPatchTool, Error, TEXT("Please ensure you've set the BuildVersionB for DiffBinaries mode and it is no more than 100 characters."));
+			return false;
+		}
+
+		if (!BuildVersionContainsOnlyAllowedSymbols(Config.BuildVersionB))
+		{
+			UE_LOG(LogBuildPatchTool, Error, TEXT("BuildVersionB string should only contain characters from the following sets a-z, A-Z, 0-9, or .+-_"));
+			return false;
+		}
+		// Get optional parameters
+		PARSE_SWITCH(InstallTagsA);
+		PARSE_SWITCH(InstallTagsB);
+		PARSE_SWITCH(OutputFile);
+		NormalizeUriFile(Config.OutputFile);
+
+		TArray<FString> CompareTagsArray;
+		ParseSwitches(TEXT("CompareTagSet="), CompareTagsArray, Switches);
+
+		for (const FString& CompareTagsList : CompareTagsArray)
+		{
+			Config.CompareTagSets.Add(ProcessTagList(CompareTagsList));
+		}
+
+		return true;
+#undef SENTENCE_END
+#undef EXPECT_NONEMPTY_CONFIG
+#undef PARSE_UNEXPECTED_SWITCH
+#undef PARSE_REQUIRED_SWITCH
+#undef _PARSE_REQUIRED_SWITCH
+#undef PARSE_SWITCH_NO_CONFIG
+#undef PARSE_SWITCH
+	}
+
+	TSet<FString> ProcessTagList(const FString& TagCommandLine) const
+	{
+		TArray<FString> TagArray;
+		TagCommandLine.ParseIntoArray(TagArray, TEXT(","), false);
+		for (FString& Tag : TagArray)
+		{
+			Tag.TrimStartAndEndInline();
+		}
+		if (TagArray.Num() == 0)
+		{
+			TagArray.Add(TEXT(""));
+		}
+		return TSet<FString>(MoveTemp(TagArray));
+	}
+private:
+	const TCHAR* CommandLine;
+	TUniquePtr<IDiffBinariesMode> Mode;
+	BuildPatchTool::EReturnCode OverrideReturnCode;
+};
+
+class FCopyBinaryToolModeWrapper : public BuildPatchTool::IToolMode
+{
+public:
+	using HelpType = BuildPatchTool::FCopyBinaryHelp;
+
+	FCopyBinaryToolModeWrapper(const BuildPatchTool::FToolModeConfiguration& Configuration)
+		: CommandLine(Configuration.CommandLine)
+		, OverrideReturnCode(BuildPatchTool::EReturnCode::OK)
+	{
+		FCopyBinaryConfig Config;
+		if (!ProcessCommandline(Config))
+		{
+			OverrideReturnCode = BuildPatchTool::EReturnCode::ArgumentProcessingError;
+		}
+		Mode = TUniquePtr<ICopyBinaryMode>(Configuration.BpMcpModule.CreateCopyBinaryMode(Configuration.BpsInterface, MoveTemp(Config)));
+	}
+
+public:
+	BuildPatchTool::EReturnCode Execute() override
+	{
+		return ExecuteHelper(Mode, OverrideReturnCode);
+	}
+	bool IsOnline() const override
+	{
+		return true;
+	}
+private:
+	bool ProcessCommandline(FCopyBinaryConfig& Config)
+	{
+#define PARSE_SWITCH(Switch) ParseSwitch(TEXT(#Switch) TEXT("="), Config.Switch, Switches)
+		TArray<FString> Tokens, Switches;
+		FCommandLine::Parse(CommandLine, Tokens, Switches);
+
+		const bool bHelp = ParseOption(TEXT("help"), Switches);
+
+		if (bHelp)
+		{
+			// If help is requested, the help mode wrapper should be used.
+			return false;
+		}
+
+		// Get all required parameters
+		if (!(PARSE_SWITCH(OrganizationId)
+			&& PARSE_SWITCH(ProductId)
+			&& PARSE_SWITCH(DestArtifactId)
+			&& PARSE_SWITCH(BuildVersion)
+			&& PARSE_SWITCH(SourceArtifactId)))
+		{
+			UE_LOG(LogBuildPatchTool, Error, TEXT("OrganizationId, ProductId, DestArtifactId, BuildVersion and SourceArtifactId are required parameters"));
+			return false;
+		}
+
+		if (Config.SourceArtifactId.IsEmpty() || Config.DestArtifactId.IsEmpty())
+		{
+			UE_LOG(LogBuildPatchTool, Error, TEXT("SourceArtifactId and DestArtifactId arguments can't be empty"));
+			return false;
+		}
+
+		if (Config.OrganizationId.IsEmpty() || Config.ProductId.IsEmpty())
+		{
+			UE_LOG(LogBuildPatchTool, Error, TEXT("Please check your OrganizationId and ProductId have been specified correctly."));
+			return false;
+		}
+
+		if (Config.BuildVersion.IsEmpty() || Config.BuildVersion.Len() > 100)
+		{
+			UE_LOG(LogBuildPatchTool, Error, TEXT("Please ensure you've set the BuildVersion for CopyBinary mode and it is no more than 100 characters."));
+			return false;
+		}
+
+		if (!BuildVersionContainsOnlyAllowedSymbols(Config.BuildVersion))
+		{
+			UE_LOG(LogBuildPatchTool, Error, TEXT("BuildVersion string should only contain characters from the following sets a-z, A-Z, 0-9, or .+-_"));
+			return false;
+		}
+
+		return true;
+#undef PARSE_SWITCH
+	}
+private:
+	const TCHAR* CommandLine;
+	TUniquePtr<ICopyBinaryMode> Mode;
+	BuildPatchTool::EReturnCode OverrideReturnCode;
+};
+
+class FDeleteBinaryToolModeWrapper : public BuildPatchTool::IToolMode
+{
+public:
+	using HelpType = BuildPatchTool::FDeleteBinaryHelp;
+
+	FDeleteBinaryToolModeWrapper(const BuildPatchTool::FToolModeConfiguration& Configuration)
+		: CommandLine(Configuration.CommandLine)
+		, OverrideReturnCode(BuildPatchTool::EReturnCode::OK)
+	{
+		FDeleteBinaryConfig Config;
+		if (!ProcessCommandline(Config))
+		{
+			OverrideReturnCode = BuildPatchTool::EReturnCode::ArgumentProcessingError;
+		}
+		Mode = TUniquePtr<IDeleteBinaryMode>(Configuration.BpMcpModule.CreateDeleteBinaryMode(Configuration.BpsInterface, MoveTemp(Config)));
+	}
+
+public:
+	BuildPatchTool::EReturnCode Execute() override
+	{
+		return ExecuteHelper(Mode, OverrideReturnCode);
+	}
+	bool IsOnline() const override
+	{
+		return true;
+	}
+private:
+	bool ProcessCommandline(FDeleteBinaryConfig& Config)
+	{
+#define PARSE_SWITCH(Switch) ParseSwitch(TEXT(#Switch) TEXT("="), Config.Switch, Switches)
+		TArray<FString> Tokens, Switches;
+		FCommandLine::Parse(CommandLine, Tokens, Switches);
+
+		const bool bHelp = ParseOption(TEXT("help"), Switches);
+		if (bHelp)
+		{
+			// If help is requested, the help mode wrapper should be used.
+			return false;
+		}
+
+		// Get all required parameters
+		if (!(PARSE_SWITCH(OrganizationId)
+			&& PARSE_SWITCH(ProductId)
+			&& PARSE_SWITCH(ArtifactId)
+			&& PARSE_SWITCH(BuildVersion)))
+		{
+			UE_LOG(LogBuildPatchTool, Error, TEXT("OrganizationId, ProductId, ArtifactId and BuildVersion are required parameters"));
+			return false;
+		}
+
+		if (Config.ArtifactId.IsEmpty() || Config.OrganizationId.IsEmpty() || Config.ProductId.IsEmpty())
+		{
+			UE_LOG(LogBuildPatchTool, Error, TEXT("Please check your OrganizationId, ProductId and ArtifactId have been specified correctly."));
+			return false;
+		}
+
+		if (Config.BuildVersion.IsEmpty() || Config.BuildVersion.Len() > 100)
+		{
+			UE_LOG(LogBuildPatchTool, Error, TEXT("Please ensure you've set the BuildVersion for DeleteBinary mode and it is no more than 100 characters."));
+			return false;
+		}
+
+		if (!BuildVersionContainsOnlyAllowedSymbols(Config.BuildVersion))
+		{
+			UE_LOG(LogBuildPatchTool, Error, TEXT("BuildVersion string should only contain characters from the following sets a-z, A-Z, 0-9, or .+-_"));
+			return false;
+		}
+
+		return true;
+#undef PARSE_SWITCH
+	}
+private:
+	const TCHAR* CommandLine;
+	TUniquePtr<IDeleteBinaryMode> Mode;
+	BuildPatchTool::EReturnCode OverrideReturnCode;
+};
+
+class FListBinariesToolModeWrapper : public BuildPatchTool::IToolMode
+{
+public:
+	using HelpType = BuildPatchTool::FListBinariesHelp;
+
+	FListBinariesToolModeWrapper(const BuildPatchTool::FToolModeConfiguration& Configuration)
+		: CommandLine(Configuration.CommandLine)
+		, OverrideReturnCode(BuildPatchTool::EReturnCode::OK)
+	{
+		FListBinariesConfig Config;
+		if (!ProcessCommandline(Config))
+		{
+			OverrideReturnCode = BuildPatchTool::EReturnCode::ArgumentProcessingError;
+		}
+		Mode = TUniquePtr<IListBinariesMode>(Configuration.BpMcpModule.CreateListBinariesMode(Configuration.BpsInterface, MoveTemp(Config)));
+	}
+
+public:
+	BuildPatchTool::EReturnCode Execute() override
+	{
+		return ExecuteHelper(Mode, OverrideReturnCode);
+	}
+	bool IsOnline() const override
+	{
+		return true;
+	}
+private:
+	bool ProcessCommandline(FListBinariesConfig& Config)
+	{
+#define PARSE_SWITCH(Switch) ParseSwitch(TEXT(#Switch) TEXT("="), Config.Switch, Switches)
+		TArray<FString> Tokens, Switches;
+		FCommandLine::Parse(CommandLine, Tokens, Switches);
+
+		const bool bHelp = ParseOption(TEXT("help"), Switches);
+		if (bHelp)
+		{
+			// If help is requested, the help mode wrapper should be used.
+			return false;
+		}
+
+		// Get all required parameters
+		if (!(PARSE_SWITCH(OrganizationId)
+			&& PARSE_SWITCH(ProductId)
+			&& PARSE_SWITCH(ArtifactId)))
+		{
+			UE_LOG(LogBuildPatchTool, Error, TEXT("OrganizationId, ProductId and ArtifactId are required parameters"));
+			return false;
+		}
+
+		if (Config.ArtifactId.IsEmpty() || Config.OrganizationId.IsEmpty() || Config.ProductId.IsEmpty())
+		{
+			UE_LOG(LogBuildPatchTool, Error, TEXT("Please check your OrganizationId, ProductId and ArtifactId have been specified correctly."));
+			return false;
+		}
+
+		// Optional parameters
+		if (ParseSwitch<uint32>(TEXT("Num="), (uint32&)Config.Num, Switches))
+		{
+			// handle case where value string was empty
+			if (Config.Num == 0)
+			{
+				Config.Num = -1;
+				UE_LOG(LogBuildPatchTool, Warning, TEXT("Requested -Num=0 is not an allowed value. Please update your args to be within range. Continuing with -1."));
+			}
+		}
+
+		Config.bOnlyLabeled = ParseOption(TEXT("onlylabeled"), Switches);
+		PARSE_SWITCH(OutputFile);
+
+		return true;
+#undef PARSE_SWITCH
+	}
+private:
+	const TCHAR* CommandLine;
+	TUniquePtr<IListBinariesMode> Mode;
+	BuildPatchTool::EReturnCode OverrideReturnCode;
+};
+
+class FUnlabelBinaryToolModeWrapper : public BuildPatchTool::IToolMode
+{
+public:
+	using HelpType = BuildPatchTool::FUnlabelBinaryHelp;
+
+	FUnlabelBinaryToolModeWrapper(const BuildPatchTool::FToolModeConfiguration& Configuration)
+		: CommandLine(Configuration.CommandLine)
+		, OverrideReturnCode(BuildPatchTool::EReturnCode::OK)
+	{
+		FUnlabelBinaryConfig Config;
+		if (!ProcessCommandLine(Config))
+		{
+			OverrideReturnCode = BuildPatchTool::EReturnCode::ArgumentProcessingError;
+		}
+		Mode = TUniquePtr<IUnlabelBinaryMode>(Configuration.BpMcpModule.CreateUnlabelBinaryMode(Configuration.BpsInterface, MoveTemp(Config)));
+	}
+
+public:
+	BuildPatchTool::EReturnCode Execute() override
+	{
+		return ExecuteHelper(Mode, OverrideReturnCode);
+	}
+	bool IsOnline() const override
+	{
+		return true;
+	}
+private:
+	bool ProcessCommandLine(FUnlabelBinaryConfig& Config)
+	{
+#define PARSE_SWITCH(Switch) ParseSwitch(TEXT(#Switch) TEXT("="), Config.Switch, Switches)
+		TArray<FString> Tokens, Switches;
+		FCommandLine::Parse(CommandLine, Tokens, Switches);
+
+		const bool bHelp = ParseOption(TEXT("help"), Switches);
+		if (bHelp)
+		{
+			// If help is requested, the help mode wrapper should be used.
+			return false;
+		}
+
+		// Get all required parameters
+		if (!(PARSE_SWITCH(OrganizationId)
+			&& PARSE_SWITCH(ProductId)
+			&& PARSE_SWITCH(ArtifactId)
+			&& PARSE_SWITCH(BuildVersion)
+			&& PARSE_SWITCH(Label)
+			&& PARSE_SWITCH(Platform)))
+		{
+			UE_LOG(LogBuildPatchTool, Error, TEXT("OrganizationId, ProductId, ArtifactId, BuildVersion, Label and Platform are required parameters"));
+			return false;
+		}
+
+		if (Config.ArtifactId.IsEmpty() || Config.OrganizationId.IsEmpty() || Config.ProductId.IsEmpty())
+		{
+			UE_LOG(LogBuildPatchTool, Error, TEXT("Please check your OrganizationId, ProductId and ArtifactId have been specified correctly."));
+			return false;
+		}
+
+		if (Config.Label.IsEmpty() || Config.Platform.IsEmpty())
+		{
+			UE_LOG(LogBuildPatchTool, Error, TEXT("Label and Platform arguments can't be empty"));
+			return false;
+		}
+
+		if (Config.BuildVersion.IsEmpty() || Config.BuildVersion.Len() > 100)
+		{
+			UE_LOG(LogBuildPatchTool, Error, TEXT("Please ensure you've set the BuildVersion and it is no more than 100 characters."));
+			return false;
+		}
+
+		if (!BuildVersionContainsOnlyAllowedSymbols(Config.BuildVersion))
+		{
+			UE_LOG(LogBuildPatchTool, Error, TEXT("BuildVersion string should only contain characters from the following sets a-z, A-Z, 0-9, or .+-_"));
+			return false;
+		}
+
+		// Get optional value
+		PARSE_SWITCH(SandboxId);
+
+		return true;
+#undef PARSE_SWITCH
+	}
+private:
+	const TCHAR* CommandLine;
+	TUniquePtr<IUnlabelBinaryMode> Mode;
+	BuildPatchTool::EReturnCode OverrideReturnCode;
+};
+
+class FLabelBinaryToolModeWrapper : public BuildPatchTool::IToolMode
+{
+public:
+	using HelpType = BuildPatchTool::FLabelBinaryHelp;
+
+	FLabelBinaryToolModeWrapper(const BuildPatchTool::FToolModeConfiguration& Configuration)
+		: CommandLine(Configuration.CommandLine)
+		, OverrideReturnCode(BuildPatchTool::EReturnCode::OK)
+	{
+		FLabelBinaryConfig Config;
+		if (!ProcessCommandLine(Config))
+		{
+			OverrideReturnCode = BuildPatchTool::EReturnCode::ArgumentProcessingError;
+		}
+		Mode = TUniquePtr<ILabelBinaryMode>(Configuration.BpMcpModule.CreateLabelBinaryMode(Configuration.BpsInterface, MoveTemp(Config)));
+	}
+
+public:
+	BuildPatchTool::EReturnCode Execute() override
+	{
+		return ExecuteHelper(Mode, OverrideReturnCode);
+	}
+	bool IsOnline() const override
+	{
+		return true;
+	}
+private:
+	bool ProcessCommandLine(FLabelBinaryConfig& Config)
+	{
+#define PARSE_SWITCH(Switch) ParseSwitch(TEXT(#Switch) TEXT("="), Config.Switch, Switches)
+		TArray<FString> Tokens, Switches;
+		FCommandLine::Parse(CommandLine, Tokens, Switches);
+
+		const bool bHelp = ParseOption(TEXT("help"), Switches);
+		if (bHelp)
+		{
+			// If help is requested, the help mode wrapper should be used.
+			return false;
+		}
+
+		// Get all required parameters
+		if (!(PARSE_SWITCH(OrganizationId)
+			&& PARSE_SWITCH(ProductId)
+			&& PARSE_SWITCH(ArtifactId)
+			&& PARSE_SWITCH(BuildVersion)
+			&& PARSE_SWITCH(Label)
+			&& PARSE_SWITCH(Platform)))
+		{
+			UE_LOG(LogBuildPatchTool, Error, TEXT("OrganizationId, ProductId, ArtifactId, BuildVersion, Label and Platform are required parameters"));
+			return false;
+		}
+
+		if (Config.ArtifactId.IsEmpty() || Config.OrganizationId.IsEmpty() || Config.ProductId.IsEmpty())
+		{
+			UE_LOG(LogBuildPatchTool, Error, TEXT("Please check your OrganizationId, ProductId and ArtifactId have been specified correctly."));
+			return false;
+		}
+
+		if (Config.Label.IsEmpty() || Config.Platform.IsEmpty())
+		{
+			UE_LOG(LogBuildPatchTool, Error, TEXT("Label and Platform arguments can't be empty"));
+			return false;
+		}
+
+		if (Config.BuildVersion.IsEmpty() || Config.BuildVersion.Len() > 100)
+		{
+			UE_LOG(LogBuildPatchTool, Error, TEXT("Please ensure you've set the BuildVersion and it is no more than 100 characters."));
+			return false;
+		}
+
+		if (!BuildVersionContainsOnlyAllowedSymbols(Config.BuildVersion))
+		{
+			UE_LOG(LogBuildPatchTool, Error, TEXT("BuildVersion string should only contain characters from the following sets a-z, A-Z, 0-9, or .+-_"));
+			return false;
+		}
+
+		// Get optional value
+		PARSE_SWITCH(SandboxId);
+
+		return true;
+#undef PARSE_SWITCH
+	}
+private:
+	const TCHAR* CommandLine;
+	TUniquePtr<ILabelBinaryMode> Mode;
+	BuildPatchTool::EReturnCode OverrideReturnCode;
+};
+
+class FContentProtectionToolModeWrapper : public BuildPatchTool::IToolMode
+{
+public:
+	using HelpType = BuildPatchTool::FContentProtectionHelp;
+
+	FContentProtectionToolModeWrapper(const BuildPatchTool::FToolModeConfiguration& Configuration)
+		: CommandLine(Configuration.CommandLine)
+		, OverrideReturnCode(BuildPatchTool::EReturnCode::OK)
+	{
+		FContentProtectionConfig Config;
+		if (!ProcessCommandLine(Config))
+		{
+			OverrideReturnCode = BuildPatchTool::EReturnCode::ArgumentProcessingError;
+		}
+		Mode = TUniquePtr<IContentProtectionMode>(Configuration.BpMcpModule.CreateContentProtectionMode(Configuration.BpsInterface, MoveTemp(Config)));
+	}
+
+public:
+	BuildPatchTool::EReturnCode Execute() override
+	{
+		return ExecuteHelper(Mode, OverrideReturnCode);
+	}
+	bool IsOnline() const override
+	{
+		return true;
+	}
+private:
+	bool ProcessCommandLine(FContentProtectionConfig& Config)
+	{
+#define PARSE_SWITCH(Switch) ParseSwitch(TEXT(#Switch) TEXT("="), Config.Switch, Switches)
+		TArray<FString> Tokens, Switches;
+		FCommandLine::Parse(CommandLine, Tokens, Switches);
+
+		const bool bHelp = ParseOption(TEXT("help"), Switches);
+		if (bHelp)
+		{
+			// If help is requested, the help mode wrapper should be used.
+			return false;
+		}
+
+		// Get all required parameters
+		if (!(PARSE_SWITCH(ArtifactId)
+			&& PARSE_SWITCH(InputFile)
+			&& PARSE_SWITCH(OutputFile)
+			&& PARSE_SWITCH(Platform)
+			&& PARSE_SWITCH(ExeFileEncoding)))
+		{
+			UE_LOG(LogBuildPatchTool, Error, TEXT("ArtifactId, InputFile, OutputFile, Platform, ExeFileEncoding are required parameters."));
+			return false;
+		}
+
+		if (!Config.ExeFileEncoding.Equals(TEXT("32bit")) && !Config.ExeFileEncoding.Equals(TEXT("64bit")))
+		{
+			UE_LOG(LogBuildPatchTool, Error, TEXT("Possible values for ExeFileEncoding are \"32bit\" and \"64bit\"."));
+			return false;
+		}
+
+		if (Config.InputFile.IsEmpty() || Config.OutputFile.IsEmpty() || Config.InputFile.EndsWith(TEXT("/")) || Config.OutputFile.EndsWith((TEXT("/"))))
+		{
+			UE_LOG(LogBuildPatchTool, Error, TEXT("InputFile and OutputFile are required parameters that can't be empty nor be a directory."));
+			return false;
+		}
+
+		if (Config.InputFile == Config.OutputFile)
+		{
+			UE_LOG(LogBuildPatchTool, Error, TEXT("InputFile and OutputFile should not point to the same file."));
+			return false;
+		}
+
+		return true;
+#undef PARSE_SWITCH
+	}
+private:
+	const TCHAR* CommandLine;
+	TUniquePtr<IContentProtectionMode> Mode;
+	BuildPatchTool::EReturnCode OverrideReturnCode;
+};
+
+class FGetBinaryMetadataToolModeWrapper : public BuildPatchTool::IToolMode
+{
+public:
+	using ModeType = IGetBinaryMetadataMode;
+	using ConfigType = FGetBinaryMetadataConfig;
+	using HelpType = BuildPatchTool::FGetBinaryMetadataHelp;
+
+	FGetBinaryMetadataToolModeWrapper(const BuildPatchTool::FToolModeConfiguration& Configuration)
+		: OverrideReturnCode(BuildPatchTool::EReturnCode::OK)
+	{
+		ConfigType Config;
+		if (!ProcessCommandline(Configuration.CommandLine, Config))
+		{
+			OverrideReturnCode = BuildPatchTool::EReturnCode::ArgumentProcessingError;
+		}
+		ManifestMetadataSerializer = TUniquePtr<IManifestMetadataSerializer>(new FManifestMetadataSerializer);
+		Mode = TUniquePtr<ModeType>(Configuration.BpMcpModule.CreateGetBinaryMetadataMode(Configuration.BpsInterface, MoveTemp(Config), *ManifestMetadataSerializer));
+	}
+
+	BuildPatchTool::EReturnCode Execute() override
+	{
+		return ExecuteHelper(Mode, OverrideReturnCode);
+	}
+
+	bool IsOnline() const override
+	{
+		return true;
+	}
+private:
+	bool ProcessCommandline(const TCHAR* CommandLine, ConfigType& Config)
+	{
+		bool bParametersOk = true;
+#define PARSE_SWITCH(Switch) ParseSwitch(TEXT(#Switch "="), Config.Switch, Switches)
+#define PARSE_REQUIRED_SWITCH(Switch) \
+		if (!PARSE_SWITCH(Switch)) \
+		{ \
+			UE_LOG(LogBuildPatchTool, Error, #Switch TEXT(" is a required parameter.")); \
+			bParametersOk = false; \
+		} \
+		else if (Config.Switch.IsEmpty()) \
+		{ \
+			UE_LOG(LogBuildPatchTool, Error, #Switch TEXT(" should not be empty.")); \
+			bParametersOk = false; \
+		}
+#define PARSE_OPTIONAL_SWITCH(Switch) \
+		if (PARSE_SWITCH(Switch) && Config.Switch.IsEmpty()) \
+		{ \
+			UE_LOG(LogBuildPatchTool, Error, #Switch TEXT(" should not be empty.")); \
+			bParametersOk = false; \
+		}
+
+		TArray<FString> Tokens, Switches;
+		FCommandLine::Parse(CommandLine, Tokens, Switches);
+
+		// Get required params.
+		PARSE_REQUIRED_SWITCH(OrganizationId);
+		PARSE_REQUIRED_SWITCH(ProductId);
+		PARSE_REQUIRED_SWITCH(ArtifactId);
+		PARSE_REQUIRED_SWITCH(BuildVersion);
+		PARSE_REQUIRED_SWITCH(OutputFile);
+
+		// Get optional params.
+		PARSE_OPTIONAL_SWITCH(OutputFormat);
+		Config.bOpenOutput = ParseOption(TEXT("OpenOutput"), Switches);
+
+		// Process params.
+		if (bParametersOk)
+		{
+			NormalizeUriPath(Config.OutputFile);
+			if (Config.OutputFormat.IsEmpty())
+			{
+				Config.OutputFormat = Config.OutputFile.EndsWith(".json") ? TEXT("json") : TEXT("human");
+			}
+			else
+			{
+				// Check the format is valid.
+				EMetadataOutputFormat MetadataOutputFormat;
+				LexFromString(MetadataOutputFormat, *Config.OutputFormat);
+				if (MetadataOutputFormat == EMetadataOutputFormat::InvalidOrMax)
+				{
+					UE_LOG(LogBuildPatchTool, Error, TEXT("OutputFormat %s was not recognised."), *Config.OutputFormat);
+					bParametersOk = false;
+				}
+			}
+		}
+
+		return bParametersOk;
+#undef PARSE_OPTIONAL_SWITCH
+#undef PARSE_REQUIRED_SWITCH
+#undef PARSE_SWITCH
+	}
+
+private:
+	BuildPatchTool::EReturnCode OverrideReturnCode;
+	TUniquePtr<IManifestMetadataSerializer> ManifestMetadataSerializer;
+	TUniquePtr<ModeType> Mode;
+};
+
+class FInstallBinaryToolModeWrapper : public BuildPatchTool::IToolMode
+{
+public:
+	using HelpType = BuildPatchTool::FInstallBinaryHelp;
+
+	FInstallBinaryToolModeWrapper(const BuildPatchTool::FToolModeConfiguration& Configuration)
+		: CommandLine(Configuration.CommandLine)
+		, OverrideReturnCode(BuildPatchTool::EReturnCode::OK)
+	{
+		FInstallBinaryConfig Config;
+		if (!ProcessCommandLine(Config))
+		{
+			OverrideReturnCode = BuildPatchTool::EReturnCode::ArgumentProcessingError;
+		}
+
+		TSharedPtr<BuildPatchServices::FMessageHandler> SignedUriProvider;
+		const bool bForceNonSignedUrls = FParse::Param(FCommandLine::Get(), TEXT("forcenonsignedurls"));
+		if (!bForceNonSignedUrls)
+		{
+			FSignedUriProviderConfig SignedUriProviderConfig
+			{
+				Config.ArtifactId,
+				Config.BuildVersion,
+				Config.ProductId,
+				Config.OrganizationId,
+				TArray<FString>() /* Disallowed CDNs are not used now */,
+				false
+			};
+			SignedUriProvider = Configuration.BpMcpModule.CreateSignedUriProvider(SignedUriProviderConfig);
+		}
+
+		Mode = TUniquePtr<IInstallBinaryMode>(Configuration.BpMcpModule.CreateInstallBinaryMode(Configuration.BpsInterface, MoveTemp(Config), MoveTemp(SignedUriProvider)));
+	}
+public:
+	BuildPatchTool::EReturnCode Execute() override
+	{
+		return ExecuteHelper(Mode, OverrideReturnCode);
+	}
+
+	bool IsOnline() const override
+	{
+		return true;
+	}
+
+private:
+	bool ProcessCommandLine(FInstallBinaryConfig& Config)
+	{
+#define PARSE_SWITCH(Switch) ParseSwitch(TEXT(#Switch "="), Config.Switch, Switches)
+#define _PARSE_REQUIRED_SWITCH(Switch, ConfigSwitch, bExpect) \
+		if (ParseSwitch(TEXT(#Switch "="), ConfigSwitch, Switches) != bExpect) \
+		{ \
+			UE_LOG(LogBuildPatchTool, Error, #Switch SENTENCE_END); \
+			bParametersOk = false; \
+		}
+#define PARSE_REQUIRED_SWITCH(Switch) _PARSE_REQUIRED_SWITCH(Switch, Config.Switch, true)
+#define PARSE_UNEXPECTED_SWITCH(Switch) do{FString Tmp##Switch; _PARSE_REQUIRED_SWITCH(Switch, Tmp##Switch, false);}while(false)
+
+		TArray<FString> Tokens, Switches;
+		FCommandLine::Parse(CommandLine, Tokens, Switches);
+
+		// If help is requested, the help mode wrapper should be used.
+		if (ParseOption(TEXT("help"), Switches))
+		{
+			// If help is requested, the help mode wrapper should be used.
+			return false;
+		}
+
+#define SENTENCE_END TEXT(" is a required parameter used for downloading from Epic services.")
+		bool bParametersOk = true;
+		// Get all required online parameters, we allow the auth manager to report errors with required params for client secret.
+		PARSE_REQUIRED_SWITCH(OrganizationId);
+		PARSE_REQUIRED_SWITCH(ProductId);
+		PARSE_REQUIRED_SWITCH(BuildVersion);
+		PARSE_REQUIRED_SWITCH(ArtifactId);
+		PARSE_REQUIRED_SWITCH(InstallDir);
+#undef SENTENCE_END
+#define SENTENCE_END TEXT(" should not be used for downloading from Epic services.")
+		// Check for params that should not have been provided for online mode.
+		PARSE_UNEXPECTED_SWITCH(Manifest);
+		PARSE_UNEXPECTED_SWITCH(DownloadCloudDirectory);
+		PARSE_UNEXPECTED_SWITCH(ScanWindowSize);
+		PARSE_UNEXPECTED_SWITCH(OutputChunkSize);
+#undef SENTENCE_END
+		if (!bParametersOk)
+		{
+			return false;
+		}
+		
+		PARSE_SWITCH(InstallTags);
+
+		if (Config.InstallDir.IsEmpty() || Config.ArtifactId.IsEmpty() || Config.OrganizationId.IsEmpty() || Config.ProductId.IsEmpty())
+		{
+			UE_LOG(LogBuildPatchTool, Error, TEXT("InstallDir, ArtifactId, OrganizationId, ProductId arguments can't be empty"));
+			return false;
+		}
+
+		if (Config.BuildVersion.IsEmpty() || Config.BuildVersion.Len() > 100)
+		{
+			UE_LOG(LogBuildPatchTool, Error, TEXT("Please ensure you've set the BuildVersion for InstallBinary mode and it is no more than 100 characters."));
+			return false;
+		}
+
+		if (!BuildVersionContainsOnlyAllowedSymbols(Config.BuildVersion))
+		{
+			UE_LOG(LogBuildPatchTool, Error, TEXT("BuildVersion string should only contain characters from the following sets a-z, A-Z, 0-9, or .+-_"));
+			return false;
+		}
+
+		NormalizeUriPath(Config.InstallDir);
+		
+		if (!PARSE_SWITCH(InstallMode) && Config.InstallMode == BuildPatchServices::EInstallMode::InvalidOrMax)
+		{
+			UE_LOG(LogBuildPatchTool, Error, TEXT("Invalid parse of argument InstallMode. Check value of argument."));
+			return false;
+		}
+
+		FString ParseErrorMessage;
+		if (!ParsePairs(TEXT("EncryptionSecretId="), TEXT("EncryptionSecretKey="), Switches, Config.EncryptionSecrets, ParseErrorMessage) && !ParseErrorMessage.IsEmpty())
+		{
+			UE_LOG(LogBuildPatchTool, Error, TEXT("%s"), *ParseErrorMessage);
+			return false;
+		}
+		
+
+		return true;
+#undef PARSE_UNEXPECTED_SWITCH
+#undef PARSE_REQUIRED_SWITCH
+#undef _PARSE_REQUIRED_SWITCH
+#undef PARSE_SWITCH
+	}
+
+private:
+	const TCHAR* CommandLine;
+	TUniquePtr<IInstallBinaryMode> Mode;
+	BuildPatchTool::EReturnCode OverrideReturnCode;
+};
+
+template<typename TWrapper>
+BuildPatchTool::IToolMode* CreateOnlineToolMode(const bool bRequestedHelp, const BuildPatchTool::FToolModeConfiguration& Configuration)
+{
+	if (bRequestedHelp)
+	{
+		return new FOnlineHelpModeWrapper<typename TWrapper::HelpType>(BuildPatchTool::EReturnCode::OK);
+	}
+	else
+	{
+		return new TWrapper(Configuration);
+	}
+}
+
+namespace
+{
+	using OnlineToolModeCreateFunc = TFunction<BuildPatchTool::IToolMode* (const bool, const BuildPatchTool::FToolModeConfiguration&)>;
+}
+
+template<typename T>
+TPair<FString, OnlineToolModeCreateFunc> CreateModePair()
+{
+	return TPair<FString, OnlineToolModeCreateFunc>{ T::Name, CreateOnlineToolMode<T> };
+}
+
+BuildPatchTool::IToolMode* BuildPatchTool::FOnlineToolModeFactory::Create(const FString& ToolModeValue, const BuildPatchTool::FToolModeConfiguration& Configuration)
+{
+	// Determine whether we're running in help mode.
+	const bool bRequestedHelp = FParse::Param(Configuration.CommandLine, TEXT("help"));
+
+	static const TMap<FString, OnlineToolModeCreateFunc> Modes{
+		{ TEXT("labelbinary"), CreateOnlineToolMode<FLabelBinaryToolModeWrapper>},
+		{ TEXT("unlabelbinary"), CreateOnlineToolMode<FUnlabelBinaryToolModeWrapper>},
+		{ TEXT("listbinaries"), CreateOnlineToolMode<FListBinariesToolModeWrapper>},
+		{ TEXT("deletebinary"), CreateOnlineToolMode<FDeleteBinaryToolModeWrapper>},
+		{ TEXT("copybinary"), CreateOnlineToolMode<FCopyBinaryToolModeWrapper>},
+		{ TEXT("binarydeltaoptimise"), CreateOnlineToolMode<FBinaryDeltaOptimiseToolModeWrapper>},
+		{ TEXT("binarydeltaoptimize"), CreateOnlineToolMode<FBinaryDeltaOptimiseToolModeWrapper>},
+		{ TEXT("diffbinaries"), CreateOnlineToolMode<FDiffBinariesToolModeWrapper>},
+		{ TEXT("uploadbinary"), CreateOnlineToolMode<FUploadBinaryToolModeWrapper>},
+		{ TEXT("applycontentprotection"), CreateOnlineToolMode<FContentProtectionToolModeWrapper>},
+		{ TEXT("reuploadbinary"), CreateOnlineToolMode<FReuploadBinaryToolModeWrapper>},
+		{ TEXT("getbinarymetadata"), CreateOnlineToolMode<FGetBinaryMetadataToolModeWrapper>},
+		{ TEXT("installbinary"), CreateOnlineToolMode<FInstallBinaryToolModeWrapper>},
+	};
+
+	if (const OnlineToolModeCreateFunc* ToolModeCreator = Modes.Find(ToolModeValue))
+	{
+		return (*ToolModeCreator)(bRequestedHelp, Configuration);
+	}
+
+	// invalid, misspelled, or empty mode provided so create the generic help, which will return ok if - help was provided else return UnknownToolMode error
+	return new FOnlineHelpModeWrapper<FNoModeHelp>(bRequestedHelp ? BuildPatchTool::EReturnCode::OK : BuildPatchTool::EReturnCode::UnknownToolMode);
+}
+
+BuildPatchTool::IToolMode* BuildPatchTool::FOnlineToolModeFactory::CreateGenericHelpMode(const TCHAR* CommandLine)
+{
+	// Determine whether we're running in help mode.
+	const bool bRequestedHelp = FParse::Param(CommandLine, TEXT("help"));
+	return new FOnlineHelpModeWrapper<FNoModeHelp>(bRequestedHelp ? BuildPatchTool::EReturnCode::OK : BuildPatchTool::EReturnCode::UnknownToolMode);
+}
+#endif
diff --git a/Engine/Source/Programs/BuildPatchTool/Private/ToolModes/OnlineToolMode.h b/Engine/Source/Programs/BuildPatchTool/Private/ToolModes/OnlineToolMode.h
new file mode 100644
--- /dev/null
+++ b/Engine/Source/Programs/BuildPatchTool/Private/ToolModes/OnlineToolMode.h
@@ -0,0 +1,18 @@
+// Copyright Epic Games, Inc. All Rights Reserved.
+#pragma once
+
+#include "CoreMinimal.h"
+#include "Interfaces/ToolMode.h"
+#include "Interfaces/IBuildPatchServicesModule.h"
+
+class IBuildPatchOnlineMcp;
+
+namespace BuildPatchTool
+{
+	class FOnlineToolModeFactory
+	{
+	public:
+		static IToolMode* Create(const FString& ToolModeValue, const BuildPatchTool::FToolModeConfiguration&);
+		static IToolMode* CreateGenericHelpMode(const TCHAR* CommandLine);
+	};
+}
\ No newline at end of file
diff --git a/Engine/Source/Programs/BuildPatchTool/Private/ToolModes/ToolModesHelp.cpp b/Engine/Source/Programs/BuildPatchTool/Private/ToolModes/ToolModesHelp.cpp
--- a/Engine/Source/Programs/BuildPatchTool/Private/ToolModes/ToolModesHelp.cpp
+++ b/Engine/Source/Programs/BuildPatchTool/Private/ToolModes/ToolModesHelp.cpp
@@ -39,19 +39,22 @@ namespace BuildPatchTool
 		TEXT("This tool mode reports the changes between two existing manifest files."),
 		TEXT(""),
 		TEXT("Required arguments:"),
-		TEXT("  -mode=DiffManifests    Must be specified to launch the tool in diff manifests mode."),
-		TEXT("  -ManifestA=\"\"          Specifies in quotes the file path to the base manifest."),
-		TEXT("  -ManifestB=\"\"          Specifies in quotes the file path to the update manifest."),
+		TEXT("  -mode=DiffManifests      Must be specified to launch the tool in diff manifests mode."),
+		TEXT("  -ManifestA=\"\"            Specifies in quotes the file path to the base manifest."),
+		TEXT("  -ManifestB=\"\"            Specifies in quotes the file path to the update manifest."),
 		TEXT(""),
 		TEXT("Optional arguments:"),
-		TEXT("  -InstallTagsA=\"\"       Specifies in quotes a comma separated list of install tags used on ManifestA. You should include empty string if you want to count untagged files."),
-		TEXT("                         Leaving the parameter out will use all files."),
-		TEXT("                         -InstallTagsA=\"\" will be untagged files only."),
-		TEXT("                         -InstallTagsA=\",tag\" will be untagged files plus files tagged with 'tag'."),
-		TEXT("                         -InstallTagsA=\"tag\" will be files tagged with 'tag' only."),
-		TEXT("  -InstallTagsB=\"\"       Specifies in quotes a comma separated list of install tags used on ManifestB. Same rules apply as InstallTagsA."),
-		TEXT("  -CompareTagSet=\"\"      Specifies in quotes a comma separated list of install tags used to calculate differential statistics betweeen the manifests. Multiple lists are allowed. Same rules apply as InstallTagsA."),
-		TEXT("  -OutputFile=\"\"         Specifies in quotes the file path where the diff will be exported as a JSON object."),
+		TEXT("  -InstallTagsA=\"\"         Specifies in quotes a comma separated list of install tags used on ManifestA. You should include empty string if you want to count untagged files."),
+		TEXT("                           Leaving the parameter out will use all files."),
+		TEXT("                           -InstallTagsA=\"\" will be untagged files only."),
+		TEXT("                           -InstallTagsA=\",tag\" will be untagged files plus files tagged with 'tag'."),
+		TEXT("                           -InstallTagsA=\"tag\" will be files tagged with 'tag' only."),
+		TEXT("  -InstallTagsB=\"\"         Specifies in quotes a comma separated list of install tags used on ManifestB. Same rules apply as InstallTagsA."),
+		TEXT("  -CompareTagSet=\"\"        Specifies in quotes a comma separated list of install tags used to calculate differential statistics betweeen the manifests. Multiple lists are allowed. Same rules apply as InstallTagsA."),
+		TEXT("  -OutputFile=\"\"           Specifies in quotes the file path where the diff will be exported as a JSON object."),
+		TEXT("  -PatchDescriptorPath=\"\"  If specified, will write out metadata files specifying what parts of each patched file changed."),
+		TEXT("  -OnlyPatchDescriptors    If specified with PatchDescriptorPath, DiffManifests will only emit the files and not emit diff info."),
+		TEXT("  -RequireOptimizedDelta   If specified, fail if an optimized delta file doesn't exist.")
 		TEXT(""),
 	};
 
@@ -223,7 +226,8 @@ namespace BuildPatchTool
 		TEXT(""),
 		TEXT("Optional arguments:"),
 		TEXT("  -PrevManifest=\"\"         Manifest of the binary already installed to OutputDir."),
-		TEXT("  -FileList=\"\"             Specifies in quotes, a comma separated list of files in the manifest to filter the installation for."),
+		TEXT("  -InstallTags=\"\"          Specifies a comma-separated list of install tags to use for selectively installing files. If left empty or missing, all files will be installed. Mutually exclusive with FileList."),
+		TEXT("  -FileList=\"\"             Specifies a comma separated list of files in the manifest to filter the installation for. Mutually exclusive with InstallTags."),
 		TEXT("  -CloudDirs=\"\"            Specifies in quotes, a comma separated list of cloud directory roots where chunks can be sourced."),
 		TEXT("  -RejectSymlinks          If specified, no symlinks are allowed to be used in manifest file."),
 		TEXT(""),
diff --git a/Engine/Source/Programs/BuildPatchTool/Public/Interfaces/ToolMode.h b/Engine/Source/Programs/BuildPatchTool/Public/Interfaces/ToolMode.h
--- a/Engine/Source/Programs/BuildPatchTool/Public/Interfaces/ToolMode.h
+++ b/Engine/Source/Programs/BuildPatchTool/Public/Interfaces/ToolMode.h
@@ -5,6 +5,7 @@
 #include "Interfaces/IBuildPatchServicesModule.h"
 #include "Algo/AllOf.h"
 #include "HAL/FileManager.h"
+#include "Misc/Optional.h"
 #include "BuildPatchTool.h"
 
 #define IMPLEMENT_BPT_MODE(ModeName, ModeClass) \
@@ -72,6 +73,14 @@ namespace BuildPatchTool
 		TCHAR const * const EqualsStr = TEXT("=");
 		TCHAR const * const QuoteStr = TEXT("\"");
 
+		struct FInstallActionArgs
+		{
+		public:
+			TOptional<FString> PrevManifestUri;
+			TOptional<FString> NewManifestUri;
+			TOptional<TSet<FString>> InstallTags;
+		};
+
 	public:
 		virtual ~IToolMode() { }
 
@@ -275,6 +284,118 @@ namespace BuildPatchTool
 			return BuildPatchTool::IToolMode::ParsePairs<FString, FString, TArray<TPair<FString, FString>>>(InFirstSwitch, InSecondSwitch, Switches, OutPairs, ErrorMessage);
 		}
 
+		bool ParseInstallActions(const TCHAR* InPrevManifestSwitch, const TCHAR* InNewManifestSwitch, const TCHAR* InTagsSwitch, const TArray<FString>& Switches, TArray<FInstallActionArgs>& OutActions, FString& ErrorMessage)
+		{
+			TOptional<FString> PrevManifest;
+			TOptional<FString> NewManifest;
+			TOptional<FString> InstallTags;
+			// Using some lambda functions to reduce copy-paste code noise.
+			auto HasEmptyValue = [](const TOptional<FString>& Val) -> bool
+				{
+					return Val.IsSet() && Val.GetValue().IsEmpty();
+				};
+			auto HasNonEmptyValue = [](const TOptional<FString>& Val) -> bool
+				{
+					return Val.IsSet() && !Val.GetValue().IsEmpty();
+				};
+			auto PushAction = [&]()
+				{
+					// Each action requires at least one manifest path, if we enter this function having only parsed empty strings or install tags alone
+					// then the commandline is not valid.
+					if (!HasNonEmptyValue(PrevManifest) && !HasNonEmptyValue(NewManifest))
+					{
+						ErrorMessage = FString::Printf(TEXT("Invalid format for install action arguments. Check the %s, %s, and %s inputs are correct and properly ordered."),
+							InPrevManifestSwitch,
+							InNewManifestSwitch,
+							InTagsSwitch);
+					}
+					// If provided, URIs should not be epty strings.
+					else if (HasEmptyValue(PrevManifest))
+					{
+						ErrorMessage = FString::Printf(TEXT("%s should not be provided empty."),
+							InPrevManifestSwitch);
+					}
+					else if (HasEmptyValue(NewManifest))
+					{
+						ErrorMessage = FString::Printf(TEXT("%s should not be provided empty."),
+							InNewManifestSwitch);
+					}
+					else
+					{
+						TSet<FString> TagSet;
+						// If tags were not provided, this is interpreted as empty set (resulting in all files).
+						// An empty string for tags on commandline interprets as a set with just an empty string in it (resulting in untagged files only).
+						if (InstallTags.IsSet())
+						{
+							ParseValue(InstallTags.GetValue(), TagSet); // Always returns true
+							InstallTags.Reset();
+						}
+						OutActions.Add({ MoveTemp(PrevManifest), MoveTemp(NewManifest), MoveTemp(TagSet) });
+						// MoveTemp above does not reset a TOptional<FString> state, it just leaves it set as empty string.
+						PrevManifest.Reset();
+						NewManifest.Reset();
+						// InstallTags is reset above.
+					}
+				};
+			auto TakeComponent = [&PushAction, this](const FString& InputString, TOptional<FString>& OutVal)
+				{
+					// If we already had a this switch, this is the start of a new action, so push the previous one.
+					if (OutVal.IsSet())
+					{
+						PushAction();
+					}
+					OutVal = InputString;
+				};
+			// Loop through all of the switches to find each action triplet. Not all args are requred per action.
+			// PushAction(...) handles making sure we got at least one manifest.
+			for (int32 SwitchIndex = 0; SwitchIndex < Switches.Num() && ErrorMessage.IsEmpty(); SwitchIndex++)
+			{
+				// Grab and trim the switch value
+				const FString& InputSwitch = Switches[SwitchIndex];
+				FString ValStr;
+				InputSwitch.Split(IToolMode::EqualsStr, nullptr, &ValStr);
+				ValStr = ValStr.TrimQuotes();
+
+				// Check for PrevManifest
+				if (InputSwitch.StartsWith(InPrevManifestSwitch))
+				{
+					NormalizeUriFile(ValStr);
+					TakeComponent(ValStr, PrevManifest);
+				}
+				// Check for NewManifest
+				else if (InputSwitch.StartsWith(InNewManifestSwitch))
+				{
+					NormalizeUriFile(ValStr);
+					TakeComponent(ValStr, NewManifest);
+				}
+				// Check for InstallTags
+				else if (InputSwitch.StartsWith(InTagsSwitch))
+				{
+					TakeComponent(ValStr, InstallTags);
+				}
+			}
+			// Check for the last action that was not pushed during the loop.
+			if (PrevManifest.IsSet() || NewManifest.IsSet() || InstallTags.IsSet())
+			{
+				PushAction();
+			}
+			// At lease one action is required to be successful.
+			if (ErrorMessage.IsEmpty() && OutActions.Num() == 0)
+			{
+				ErrorMessage = FString::Printf(TEXT("At least one install action is required. Check you have provided the %s, %s, and %s arguments necessary."),
+					InPrevManifestSwitch,
+					InNewManifestSwitch,
+					InTagsSwitch);
+			}
+			return ErrorMessage.IsEmpty();
+		}
+
+		// Helper to run ParseInstallActions with default argument names.
+		bool ParseInstallActions(const TArray<FString>& Switches, TArray<FInstallActionArgs>& OutActions, FString& ErrorMessage)
+		{
+			return ParseInstallActions(TEXT("PrevManifest="), TEXT("Manifest="), TEXT("InstallTags="), Switches, OutActions, ErrorMessage);
+		}
+
 		bool ParseValue(const FString& ValueIn, FString& ValueOut)
 		{
 			ValueOut = ValueIn.TrimQuotes();
diff --git a/Engine/Source/Runtime/Online/BuildPatchServices/Private/BuildPatchFileConstructor.cpp b/Engine/Source/Runtime/Online/BuildPatchServices/Private/BuildPatchFileConstructor.cpp
--- a/Engine/Source/Runtime/Online/BuildPatchServices/Private/BuildPatchFileConstructor.cpp
+++ b/Engine/Source/Runtime/Online/BuildPatchServices/Private/BuildPatchFileConstructor.cpp
@@ -9,10 +9,12 @@
 #include "HAL/RunnableThread.h"
 #include "Hash/xxhash.h"
 #include "Misc/ConfigCacheIni.h"
+#include "Misc/Optional.h"
 #include "Misc/Paths.h"
 #include "Misc/ScopeLock.h"
 #include "Misc/ScopeRWLock.h"
 #include "Misc/OutputDeviceRedirector.h"
+#include "Misc/ConfigCacheIni.h"
 #include "BuildPatchServicesPrivate.h"
 #include "Interfaces/IBuildInstaller.h"
 #include "Data/ChunkData.h"
@@ -84,25 +86,16 @@ namespace FileConstructorHelpers
 		}
 	}
 
-	bool CheckRemainingDiskSpace(const FString& InstallDirectory, uint64 RemainingBytesRequired, uint64& OutAvailableDiskSpace)
+	TOptional<FDiskSpaceInfo> GetRemainingDiskSpace(const FString& InstallDirectory, uint64 RemainingBytesRequired)
 	{
 		bool bContinueConstruction = true;
 		uint64 TotalSize = 0;
-		OutAvailableDiskSpace = 0;
-		if (FPlatformMisc::GetDiskTotalAndFreeSpace(InstallDirectory, TotalSize, OutAvailableDiskSpace))
+		uint64 AvailableDiskSpace = 0;
+		if (!FPlatformMisc::GetDiskTotalAndFreeSpace(InstallDirectory, TotalSize, AvailableDiskSpace))
 		{
-			if (OutAvailableDiskSpace < RemainingBytesRequired)
-			{
-				bContinueConstruction = false;
-			}
+			return TOptional<FDiskSpaceInfo>{};
 		}
-		else
-		{
-			// If we can't get the disk space free then the most likely reason is the drive is no longer around...
-			bContinueConstruction = false;
-		}
-
-		return bContinueConstruction;
+		return TOptional<FDiskSpaceInfo>{FDiskSpaceInfo{ AvailableDiskSpace, AvailableDiskSpace >= RemainingBytesRequired }};
 	}
 
 	uint64 CalculateRequiredDiskSpace(const FBuildPatchAppManifestPtr& CurrentManifest, const FBuildPatchAppManifestRef& BuildManifest, const EInstallMode& InstallMode, const TSet<FString>& InInstallTags)
@@ -1871,6 +1864,7 @@ bool FBuildPatchFileConstructor::HarvestChunksForCompletedFile(const FString& In
 	{
 		FRWScopeLock ChunkLock(ChunkLocationsLock, SLT_ReadOnly);
 
+		
 		int32 CurrentUsageIndex = ChunkReferenceTracker->GetCurrentUsageIndex();
 		for (const FGuid& FileChunk : FileChunks)
 		{
@@ -2091,15 +2085,18 @@ void FBuildPatchFileConstructor::Run()
 		ResumeData.InitResume();
 
 		// Check the stage of resumed files and remove any incompatible files.
-		if (ResumeData.bHasResumeData && !Configuration.bUseSHA1StageFilenames)
+		if (ResumeData.bHasResumeData)
 		{
-			for (const FString& FileToConstruct : Configuration.ConstructList)
+			if (!Configuration.bUseSHA1StageFilenames)
 			{
-				const bool bFileIncompatible = ResumeData.CheckFile(FileToConstruct);
-				if (bFileIncompatible)
+				for (const FString& FileToConstruct : Configuration.ConstructList)
 				{
-					UE_LOG(LogBuildPatchFileConstructor, Log, TEXT("Deleting incompatible stage file %s"), *FileToConstruct);
-					FileSystem->DeleteFile(*(Configuration.StagingDirectory / FileToConstruct));
+					const bool bFileIncompatible = ResumeData.CheckFile(FileToConstruct);
+					if (bFileIncompatible)
+					{
+						UE_LOG(LogBuildPatchFileConstructor, Log, TEXT("Deleting incompatible stage file %s"), *FileToConstruct);
+						FileSystem->DeleteFile(*(Configuration.StagingDirectory / FileToConstruct));
+					}
 				}
 			}
 
@@ -2108,14 +2105,14 @@ void FBuildPatchFileConstructor::Run()
 				UE_LOG(LogBuildPatchServices, Log, TEXT("Removing incomplete entries from resume data: construct in memory"));
 				ResumeData.FilesStage = ResumeData.FilesStage.FilterByPredicate([](const TPair<FString, EFileStage>& Entry) { return Entry.Value != EFileStage::Incomplete; });
 			}
-		}
 
-		// Uninstall should not have any files to construct
-		const bool bIsUninstall = Configuration.ConstructList.IsEmpty();
-		if (ResumeData.bHasResumeData && !bIsUninstall)
-		{
-			// For uninstall, we will remove staged files at an upper level
-			RemoveUnnecessaryStagedFiles(Configuration.StagingDirectory);
+			// Uninstall should not have any files to construct
+			const bool bIsUninstall = Configuration.ConstructList.IsEmpty();
+			if (!bIsUninstall)
+			{
+				// For uninstall, we will remove staged files at an upper level
+				RemoveUnnecessaryStagedFiles(Configuration.StagingDirectory);
+			}
 		}
 
 		// Save for started versions
@@ -2368,6 +2365,9 @@ struct FBatchState
 
 	// Set by the completing threads when the batch is done.
 	std::atomic<bool> bIsFinished = true;
+
+	// Number of chunk parts remaining in the tracker at the moment of batch read completion.
+	std::atomic<int32> RemainingChunkCount = TNumericLimits<int32>::Max();
 	
 	FBatchState()
 	{
@@ -2462,7 +2462,7 @@ void FBuildPatchFileConstructor::RequestCompletedFn(const FGuid& Guid, bool bAbo
 			Request->File->Progress += TotalToDestinationBuffer;
 			Request->File->ProgressLock.Unlock();
 		}
-		
+		Request->Batch->RemainingChunkCount.store(RemainingChunkCount, std::memory_order_relaxed);
 		if (Request->Batch->PendingRequestCount.fetch_add(-1, std::memory_order_relaxed) == 1)
 		{
 			Request->Batch->bIsFinished.store(true, std::memory_order_release);
@@ -2812,6 +2812,44 @@ void FBuildPatchFileConstructor::StartReadBatch(FFileConstructionState& CurrentF
 	CurrentFile.NextChunkPartToRead = EndChunkIdx;
 }
 
+
+void FBuildPatchFileConstructor::StartTruncateChunkDbForCompletedBatch(const FBatchState& Batch)
+{
+	UE_LOG(LogBuildPatchServices, Verbose, TEXT("Starting async truncating request"));
+	checkf(Batch.bIsWriting && Batch.bIsFinished || Batch.bIsEmptyFileSentinel, TEXT("Batch must be finished or be an empty file sentinel before truncating chunkdb"));
+	for (const TPair<FGuid, FRequestInfo>& RequestPair : Batch.Requests)
+	{
+		const FRequestInfo& Request = RequestPair.Value;
+		if (Request.bSourceIsMemoryStore)
+		{
+			// Doesn't make sense to continue for data cached in memory.
+			continue;
+		}
+
+		const EConstructorChunkLocation* ChunkLocationPtr = ChunkLocations.Find(Request.Guid);
+		if (ChunkLocationPtr == nullptr)
+		{
+			continue;
+		}
+
+		EConstructorChunkLocation ChunkLocation = *ChunkLocationPtr;
+		if (*ChunkLocationPtr != EConstructorChunkLocation::ChunkDb)
+		{
+			continue; // Only chunk db sources can be truncated.
+		}
+
+
+		IConstructorChunkSource::FRequestProcessFn TruncateTask = [this, BatchRemainingChunkCount = Batch.RemainingChunkCount.load(std::memory_order_acquire)](bool bIsAborted) {
+			if (!bIsAborted)
+			{
+				ChunkDbSource->TryTruncateFiles(BatchRemainingChunkCount);
+			}
+		};
+		// Execute the truncating on the same thread that reads chunks from chunkdb to make archive access lock-free.
+		QueueGenericThreadTask(ThreadAssignments[ChunkLocation], MoveTemp(TruncateTask));
+	}
+}
+
 void FBuildPatchFileConstructor::WakeUpDispatch()
 {
 	WakeUpDispatchThreadEvent->Trigger();
@@ -2931,6 +2969,12 @@ bool FBuildPatchFileConstructor::HandleInitialDiskSizeCheck(const FFileManifest&
 	{
 		bInitialDiskSizeCheck = true;
 
+		if (Configuration.bInstallToMemory || Configuration.bSkipInitialDiskSizeCheck)
+		{
+			UE_LOG(LogBuildPatchServices, Log, TEXT("Initial Disk Size check skipped due to the flags: bInstallToMemory: %s, bSkipInitialDiskSizeCheck: %s"), *LexToString(Configuration.bInstallToMemory), *LexToString(Configuration.bSkipInitialDiskSizeCheck));
+			return true;
+		}
+
 		// Normal operation can just use the classic calculation
 		uint64 LocalDiskSpaceRequired = 0;
 
@@ -2946,27 +2990,18 @@ bool FBuildPatchFileConstructor::HandleInitialDiskSizeCheck(const FFileManifest&
 			LocalDiskSpaceRequired = CalculateInProgressDiskSpaceRequired(FileManifest, StartPosition);
 		}
 
-		bool bHaveSufficientFreeSpace = false;
-		bool bGotValidFreeSpace = false;
-		uint64 LocalDiskSpaceAvailable = 0;
+		const TOptional<FileConstructorHelpers::FDiskSpaceInfo> DiskSpaceInfo = FileConstructorHelpers::GetRemainingDiskSpace(Configuration.InstallDirectory, LocalDiskSpaceRequired);
+		if (DiskSpaceInfo.IsSet())
 		{
-			uint64 TotalSize = 0;
-			uint64 AvailableSpace = 0;
-			bGotValidFreeSpace = FPlatformMisc::GetDiskTotalAndFreeSpace(Configuration.InstallDirectory, TotalSize, AvailableSpace);
-			if (bGotValidFreeSpace)
-			{
-				bHaveSufficientFreeSpace = AvailableSpace > LocalDiskSpaceRequired;
-				LocalDiskSpaceAvailable = AvailableSpace;
-				BackingStore->SetDynamicDiskSpaceHeadroom(LocalDiskSpaceRequired);
-			}
+			BackingStore->SetDynamicDiskSpaceHeadroom(LocalDiskSpaceRequired);
 		}
+		UE_LOG(LogBuildPatchServices, Display, TEXT("Initial Disk Sizes: Required: %s Available: %s"), *FormatNumber(LocalDiskSpaceRequired), DiskSpaceInfo.IsSet() ? *FormatNumber(DiskSpaceInfo->AvailableDiskSpace) : TEXT("<failed>"));
 
-		UE_LOG(LogBuildPatchServices, Display, TEXT("Initial Disk Sizes: Required: %s Available: %s"), *FormatNumber(LocalDiskSpaceRequired), bGotValidFreeSpace ? *FormatNumber(LocalDiskSpaceAvailable) : TEXT("<failed>"));
-
+		const uint64 LocalDiskSpaceAvailable = DiskSpaceInfo.IsSet() ? DiskSpaceInfo->AvailableDiskSpace : 0;
 		AvailableDiskSpace.store(LocalDiskSpaceAvailable, std::memory_order_release);
 		RequiredDiskSpace.store(LocalDiskSpaceRequired, std::memory_order_release);	
 
-		if (!bGotValidFreeSpace || !bHaveSufficientFreeSpace)
+		if (!DiskSpaceInfo.IsSet() || !DiskSpaceInfo->bHaveSufficientFreeSpace)
 		{
 			UE_LOG(LogBuildPatchServices, Error, TEXT("Insufficient disk space on initial check."));
 			InstallerError->SetError(
@@ -3277,9 +3312,10 @@ void FBuildPatchFileConstructor::CompleteConstructedFile(FFileConstructionState&
 		CurrentFile.SetAttributes();
 	}
 
-	if (CurrentFile.bSuccess)
+	RemainingChunkCount = ChunkReferenceTracker->GetRemainingChunkCount();
+	if (CurrentFile.bSuccess && ChunkDbSource)
 	{
-		ChunkDbSource->ReportFileCompletion(ChunkReferenceTracker->GetRemainingChunkCount());
+		ChunkDbSource->ReportFileCompletion(RemainingChunkCount);
 	}
 
 	FileConstructorStat->OnFileCompleted(CurrentFile.BuildFilename, CurrentFile.bSuccess);
@@ -3304,11 +3340,9 @@ void FBuildPatchFileConstructor::CompleteConstructedFile(FFileConstructionState&
 				// tell them we failed, and to verify disk space. It's not useful to share the actual value we
 				// get back from the operating system since we are either in a write media failure mode, or so
 				// close to the edge of free space that values can look like we have enough, but don't.
-				uint64 TotalSize, AvailableSpace;
-				bool bGotFreeSpace = FPlatformMisc::GetDiskTotalAndFreeSpace(Configuration.InstallDirectory, TotalSize, AvailableSpace);
 
 				uint64 LocalDiskSpaceRequired;
-				if (Configuration.bDeleteChunkDBFilesAfterUse)
+				if (Configuration.bDeleteChunkDBFilesAfterUse && ChunkDbSource)
 				{
 					LocalDiskSpaceRequired = CalculateDiskSpaceRequirementsWithDeleteDuringInstall(CurrentFile.Progress);
 				}
@@ -3317,7 +3351,8 @@ void FBuildPatchFileConstructor::CompleteConstructedFile(FFileConstructionState&
 					LocalDiskSpaceRequired = CalculateInProgressDiskSpaceRequired(*CurrentFile.FileManifest, CurrentFile.Progress);
 				}
 
-				UE_LOG(LogBuildPatchServices, Error, TEXT("Write failure during installation. FreeSpace: Valid=%s, Available=%llu"), bGotFreeSpace ? TEXT("true") : TEXT("false"), AvailableSpace);
+				const TOptional<FileConstructorHelpers::FDiskSpaceInfo> DiskSpaceInfo = FileConstructorHelpers::GetRemainingDiskSpace(Configuration.InstallDirectory, LocalDiskSpaceRequired);
+				UE_LOG(LogBuildPatchServices, Error, TEXT("Write failure during installation. Required: %s Available: %s"), *FormatNumber(LocalDiskSpaceRequired), DiskSpaceInfo.IsSet() ? *FormatNumber(DiskSpaceInfo->AvailableDiskSpace) : TEXT("<failed>"));
 
 				const FNumberFormattingOptions FormatOptions = FNumberFormattingOptions()
 					.SetMinimumFractionalDigits(2)
@@ -3687,6 +3722,10 @@ void FBuildPatchFileConstructor::ConstructFiles(FResumeData& ResumeData)
 						// conversion.
 						File->ConstructionError = EConstructionError::FailedWrite;
 					}
+					else
+					{
+						StartTruncateChunkDbForCompletedBatch(*Batches[0]);
+					}
 				}
 
 				Batches.RemoveAt(0);
diff --git a/Engine/Source/Runtime/Online/BuildPatchServices/Private/BuildPatchFileConstructor.h b/Engine/Source/Runtime/Online/BuildPatchServices/Private/BuildPatchFileConstructor.h
--- a/Engine/Source/Runtime/Online/BuildPatchServices/Private/BuildPatchFileConstructor.h
+++ b/Engine/Source/Runtime/Online/BuildPatchServices/Private/BuildPatchFileConstructor.h
@@ -90,7 +90,7 @@ namespace BuildPatchServices
 
 		IBuildInstallerSharedContext* SharedContext = nullptr;
 
-		// See comments in installer config
+		// See comments in installer config (FBuildInstallerConfiguration)
 		bool bInstallToMemory = false;
 		bool bConstructInMemory = false;
 		bool bSkipInitialDiskSizeCheck = false;
@@ -379,6 +379,9 @@ namespace BuildPatchServices
 		void OpenFileToConstruct(FFileConstructionState& CurrentFile);
 		bool HandleInitialDiskSizeCheck(const FFileManifest& CurrentFileManifest, int64 BytesInToCurrentFile);
 		bool HarvestChunksForCompletedFile(const FString& InstallationDirectory, const FString& FilenameInManifest);
+		// Starts an asynchronous truncation task in the chunk reading thread, which guarantees single-threaded access to chunkdb.
+		void StartTruncateChunkDbForCompletedBatch(const FBatchState& Batch);
+
 		void CompleteConstructedFile(FFileConstructionState& CurrentFile, FResumeData& ResumeData);
 
 		// Determine and return the file stage based on its existence and file size.
@@ -518,5 +521,10 @@ namespace BuildPatchServices
 namespace FileConstructorHelpers
 {
 	uint64 CalculateRequiredDiskSpace(const FBuildPatchAppManifestPtr& CurrentManifest, const FBuildPatchAppManifestRef& BuildManifest, const BuildPatchServices::EInstallMode& InstallMode, const TSet<FString>& InstallTags);
-	bool CheckRemainingDiskSpace(const FString& InstallDirectory, uint64 RemainingBytesRequired, uint64& OutAvailableDiskSpace);
+	struct FDiskSpaceInfo
+	{ 
+		uint64 AvailableDiskSpace = 0;
+		bool bHaveSufficientFreeSpace = false;
+	};
+	TOptional<FDiskSpaceInfo> GetRemainingDiskSpace(const FString& InstallDirectory, uint64 RemainingBytesRequired);
 }
diff --git a/Engine/Source/Runtime/Online/BuildPatchServices/Private/BuildPatchGeneration.cpp b/Engine/Source/Runtime/Online/BuildPatchServices/Private/BuildPatchGeneration.cpp
--- a/Engine/Source/Runtime/Online/BuildPatchServices/Private/BuildPatchGeneration.cpp
+++ b/Engine/Source/Runtime/Online/BuildPatchServices/Private/BuildPatchGeneration.cpp
@@ -535,7 +535,7 @@ namespace BuildPatchServices
 
 		// Enumerate Chunks.
 		const FDateTime Cutoff = Configuration.bShouldHonorReuseThreshold ? FDateTime::UtcNow() - FTimespan::FromDays(Configuration.DataAgeThreshold) : FDateTime::MinValue();
-		TUniquePtr<ICloudEnumeration> CloudEnumeration(FCloudEnumerationFactory::Create(Configuration.CloudDirectory, Cutoff, Configuration.FeatureLevel, StatsCollector.Get()));
+		TUniquePtr<ICloudEnumeration> CloudEnumeration(FCloudEnumerationFactory::Create(Configuration.CloudDirectory, Cutoff, Configuration.ReuseManifestPredicate, Configuration.FeatureLevel, StatsCollector.Get()));
 
 		// Sort the files accordingly.
 		static_assert((int32)EFileSortOrder::InvalidOrMax == 2, "Please implement additional file sort orders.");
diff --git a/Engine/Source/Runtime/Online/BuildPatchServices/Private/BuildPatchInstaller.cpp b/Engine/Source/Runtime/Online/BuildPatchServices/Private/BuildPatchInstaller.cpp
--- a/Engine/Source/Runtime/Online/BuildPatchServices/Private/BuildPatchInstaller.cpp
+++ b/Engine/Source/Runtime/Online/BuildPatchServices/Private/BuildPatchInstaller.cpp
@@ -63,6 +63,7 @@
 #include "Installer/Statistics/InstallChunkSourceStatistics.h"
 #include "Installer/Statistics/MemoryChunkStoreStatistics.h"
 #include "Installer/Statistics/VerifierStatistics.h"
+#include "Installer/Statistics/InstallActionStatistics.h"
 #include "Installer/UriProvider.h"
 #include "Installer/Verifier.h"
 
@@ -326,10 +327,13 @@ namespace InstallerHelpers
 		UE_LOG(LogBPSInstallerConfig, Log, TEXT("Build Config: InstallMode: %s"), LexToString(InstallerConfiguration.InstallMode));
 		UE_LOG(LogBPSInstallerConfig, Log, TEXT("Build Config: VerifyMode: %s"), LexToString(InstallerConfiguration.VerifyMode));
 		UE_LOG(LogBPSInstallerConfig, Log, TEXT("Build Config: DeltaPolicy: %s"), LexToString(InstallerConfiguration.DeltaPolicy));
-		UE_LOG(LogBPSInstallerConfig, Log, TEXT("Build Config: bRunRequiredPrereqs: %s"), (InstallerConfiguration.bRunRequiredPrereqs) ? TEXT("true") : TEXT("false"));
-		UE_LOG(LogBPSInstallerConfig, Log, TEXT("Build Config: bAllowConcurrentExecution: %s"), (InstallerConfiguration.bAllowConcurrentExecution) ? TEXT("true") : TEXT("false"));
+		UE_LOG(LogBPSInstallerConfig, Log, TEXT("Build Config: bRunRequiredPrereqs: %s"), *LexToString(InstallerConfiguration.bRunRequiredPrereqs));
+		UE_LOG(LogBPSInstallerConfig, Log, TEXT("Build Config: bAllowConcurrentExecution: %s"), *LexToString(InstallerConfiguration.bAllowConcurrentExecution));
 		UE_LOG(LogBPSInstallerConfig, Log, TEXT("Build Config: DownloadRateLimitBps: %llu"), InstallerConfiguration.DownloadRateLimitBps);
-		UE_LOG(LogBPSInstallerConfig, Log, TEXT("Build Config: bRejectSymlinks: %s"), (InstallerConfiguration.bRejectSymlinks) ? TEXT("true") : TEXT("false"));
+		UE_LOG(LogBPSInstallerConfig, Log, TEXT("Build Config: bRejectSymlinks: %s"), *LexToString(InstallerConfiguration.bRejectSymlinks));
+		UE_LOG(LogBPSInstallerConfig, Log, TEXT("Build Config: bConstructFilesInMemory: %s"), *LexToString(InstallerConfiguration.bConstructFilesInMemory));
+		UE_LOG(LogBPSInstallerConfig, Log, TEXT("Build Config: bInstallToMemory: %s"), *LexToString(InstallerConfiguration.bInstallToMemory));
+		UE_LOG(LogBPSInstallerConfig, Log, TEXT("Build Config: bSkipInitialDiskSizeCheck: %s"), *LexToString(InstallerConfiguration.bSkipInitialDiskSizeCheck));
 	}
 
 	TSet<FGuid> GetMultipleReferencedChunks(IBuildManifestSet* ManifestSet)
@@ -459,6 +463,8 @@ namespace BuildPatchServices
 
 	static uint64 DetermineInstallMaxDiskSizeIfDeletingChunkDbs(const TArray<FString>& ChunkDbFiles, IBuildManifestSet* ManifestSet, IFileSystem* FileSystem, const FString& InstallDirectory, const TArray<FString>& CorruptFiles, EInstallMode InstallMode);
 
+	static FString FormatNumber(uint64 Value) { return FText::AsNumber(Value).ToString(); }
+
 	namespace ConstructionHelpers
 	{
 		FBuildInstallerConfiguration&& ProcessOverrides(FBuildInstallerConfiguration&& Configuration)
@@ -615,14 +621,15 @@ namespace BuildPatchServices
 		, InstallerError(FInstallerErrorFactory::Create())
 		, Analytics(MoveTemp(InAnalytics))
 		, InstallerAnalytics(FInstallerAnalyticsFactory::Create(Analytics.Get()))
+		, InstallActionStatistics(FInstallActionStatisticsFactory::Create(InstallerActions))
 		, FileOperationTracker(Configuration.bTrackFileOperations ? FFileOperationTrackerFactory::Create(FTSTicker::GetCoreTicker()) : FFileOperationTrackerFactory::CreateNull())
 		, DownloadSpeedRecorder(FSpeedRecorderFactory::Create())
 		, DiskReadSpeedRecorder(FSpeedRecorderFactory::Create())
 		, DiskWriteSpeedRecorder(FSpeedRecorderFactory::Create())
 		, ChunkDataSizeProvider(FChunkDataSizeProviderFactory::Create(ConstructionHelpers::ToManifestArray(InstallerActions)))
-		, DownloadServiceStatistics(FDownloadServiceStatisticsFactory::Create(DownloadSpeedRecorder.Get(), ChunkDataSizeProvider.Get(), InstallerAnalytics.Get()))
+		, DownloadServiceStatistics(FDownloadServiceStatisticsFactory::Create(DownloadSpeedRecorder, ChunkDataSizeProvider, InstallerAnalytics))
 		, ChunkDbChunkSourceStatistics(FChunkDbChunkSourceStatisticsFactory::Create(DiskReadSpeedRecorder.Get(), FileOperationTracker.Get()))
-		, InstallChunkSourceStatistics(FInstallChunkSourceStatisticsFactory::Create(DiskReadSpeedRecorder.Get(), InstallerAnalytics.Get(), FileOperationTracker.Get()))
+		, InstallChunkSourceStatistics(FInstallChunkSourceStatisticsFactory::Create(DiskReadSpeedRecorder.Get(), InstallerAnalytics.Get(), InstallActionStatistics.Get(), FileOperationTracker.Get()))
 		, CloudChunkSourceStatistics(FCloudChunkSourceStatisticsFactory::Create(InstallerAnalytics.Get(), &BuildProgress, FileOperationTracker.Get()))
 		, FileConstructorStatistics(FFileConstructorStatisticsFactory::Create(DiskReadSpeedRecorder.Get(), DiskWriteSpeedRecorder.Get(), &BuildProgress, FileOperationTracker.Get()))
 		, VerifierStatistics(FVerifierStatisticsFactory::Create(DiskReadSpeedRecorder.Get(), &BuildProgress, FileOperationTracker.Get()))
@@ -676,21 +683,25 @@ namespace BuildPatchServices
 
 	const ISpeedRecorder* FBuildPatchInstaller::GetDownloadSpeedRecorder() const
 	{
+		check(DownloadSpeedRecorder.IsValid());
 		return DownloadSpeedRecorder.Get();
 	}
 
 	const ISpeedRecorder* FBuildPatchInstaller::GetDiskReadSpeedRecorder() const
 	{
+		check(DiskReadSpeedRecorder.IsValid());
 		return DiskReadSpeedRecorder.Get();
 	}
 
 	const ISpeedRecorder* FBuildPatchInstaller::GetDiskWriteSpeedRecorder() const
 	{
+		check(DiskWriteSpeedRecorder.IsValid());
 		return DiskWriteSpeedRecorder.Get();
 	}
 
 	const IDownloadServiceStatistics* FBuildPatchInstaller::GetDownloadServiceStatistics() const
 	{
+		check(DownloadServiceStatistics.IsValid());
 		return DownloadServiceStatistics.Get();
 	}
 
@@ -1022,6 +1033,8 @@ namespace BuildPatchServices
 				{
 					InstallerAction.SetDeltaManifest(OptimisedDeltaResult.GetValue().ToSharedRef());
 				}
+				// Report delta manifest use and response code to analytics.
+				InstallActionStatistics->OnUsedOptimizedDelta(InstallerAction, true, OptimisedDelta->GetDownloadResponseCode());
 			}
 			else
 			{
@@ -1303,6 +1316,7 @@ namespace BuildPatchServices
 			BuildStats.FailureReasonText = InstallerError->GetErrorText();
 			BuildStats.FailureType = InstallerError->GetErrorType();
 			BuildStats.MaxDiskSpaceNeededWhenDeletingChunkDbsIfRequested = MaxDiskSpaceNeededWhenDeletingChunkDbsIfRequested;
+			BuildStats.ActionsStats = InstallActionStatistics->GetAllStats();
 		}
 
 		// Mark that we are done
@@ -1651,17 +1665,10 @@ namespace BuildPatchServices
 		return MaxDiskSize;
 	}
 
-
 	bool FBuildPatchInstaller::RunInstallation(TArray<FString>& CorruptFiles)
 	{
-		// We want to reset errors because we are retrying, however we ONLY want to do this if the error was caused by us!
-		// This is a bit delicate because this pointer gets hit by any thread, but we know that installer errors are latching.
-		// (except for retry).
-		if (!InstallerError->IsCancelled())
-		{
-			InstallerError->Reset();
-			InstallerError->RegisterForErrors([this]() { CancelInstall(); });
-		}
+		// Reset all errors except for canceled errors to preserve the canceled state between install retries.
+		InstallerError->ResetError(true /*bPreserveCancelled*/);
 
 		// Installation is not ran for VerifyOnly actions. This mode is read-only.
 		if (ManifestSet->IsExclusivelyVerifying())
@@ -1871,6 +1878,9 @@ namespace BuildPatchServices
 				InstallChunkSource.Reset();
 			}
 
+			// Gather statistics on expected and available chunks for each action.
+			InstallActionStatistics->CalculateChunkStatistics(InstallChunks);
+
 			// Set up our chunk location tracking.
 			int32 DownloadChunkCount = CloudChunks.Num();
 			TMap<FGuid, EConstructorChunkLocation> ChunkLocations;	
@@ -2099,10 +2109,10 @@ namespace BuildPatchServices
 
 	bool FBuildPatchInstaller::RunPreload()
 	{
-		// Reset our error and build progress
-		InstallerError.Reset(FInstallerErrorFactory::Create());
-		InstallerError->RegisterForErrors([this]() { CancelInstall(); });
+		// Reset all errors except for canceled errors to preserve the canceled state between install retries.
+		InstallerError->ResetError(true /*bPreserveCancelled*/);
 
+		// Reset build progress
 		UE_LOG(LogBuildPatchServices, Log, TEXT("Starting Preload"));
 
 		// Set weights for preload only
@@ -2227,15 +2237,15 @@ namespace BuildPatchServices
 				}
 
 				// Check we have enough disk space to proceed
-				uint64 AvailableDiskSpace = 0;
-				if (!FileConstructorHelpers::CheckRemainingDiskSpace(Configuration.InstallDirectory, RequiredDiskSpace, AvailableDiskSpace))
+				const TOptional<FileConstructorHelpers::FDiskSpaceInfo> DiskSpaceInfo = FileConstructorHelpers::GetRemainingDiskSpace(Configuration.InstallDirectory, RequiredDiskSpace);
+				if (!DiskSpaceInfo.IsSet() || !DiskSpaceInfo->bHaveSufficientFreeSpace)
 				{
-					UE_LOG(LogBuildPatchServices, Error, TEXT("Preload: Out of HDD space. Needs %llu bytes, Free %llu bytes"), RequiredDiskSpace, AvailableDiskSpace);
+					UE_LOG(LogBuildPatchServices, Display, TEXT("Preload: Initial Disk Sizes: Required: %s Available: %s"), *FormatNumber(RequiredDiskSpace), DiskSpaceInfo.IsSet() ? *FormatNumber(DiskSpaceInfo->AvailableDiskSpace) : TEXT("<failed>"));
 					InstallerError->SetError(
 						EBuildPatchInstallError::OutOfDiskSpace,
 						PreExistingChunkDbs.Num() > 0 ? DiskSpaceErrorCodes::DuringInstallation : DiskSpaceErrorCodes::InitialSpaceCheck,
 						0,
-						BuildPatchServices::GetDiskSpaceMessage(Configuration.InstallDirectory, RequiredDiskSpace, AvailableDiskSpace));
+						BuildPatchServices::GetDiskSpaceMessage(Configuration.InstallDirectory, RequiredDiskSpace, DiskSpaceInfo.IsSet() ? DiskSpaceInfo->AvailableDiskSpace : 0));
 					bChunkDbSuccess = false;
 				}
 				else
@@ -2323,7 +2333,7 @@ namespace BuildPatchServices
 					FScopeLock Lock(&ThreadLock);
 					BuildStats.NumFilesConstructed = FileConstructorStatistics->GetFilesConstructed();
 					BuildStats.ProcessRequiredDiskSpace = RequiredDiskSpace;
-					BuildStats.ProcessAvailableDiskSpace = AvailableDiskSpace;
+					BuildStats.ProcessAvailableDiskSpace = DiskSpaceInfo.IsSet() ? DiskSpaceInfo->AvailableDiskSpace : 0;
 				}
 
 				// Process some final stats.
diff --git a/Engine/Source/Runtime/Online/BuildPatchServices/Private/BuildPatchInstaller.h b/Engine/Source/Runtime/Online/BuildPatchServices/Private/BuildPatchInstaller.h
--- a/Engine/Source/Runtime/Online/BuildPatchServices/Private/BuildPatchInstaller.h
+++ b/Engine/Source/Runtime/Online/BuildPatchServices/Private/BuildPatchInstaller.h
@@ -46,6 +46,7 @@ namespace BuildPatchServices
 	class ICloudChunkSourceStatistics;
 	class IInstallChunkSourceStatistics;
 	class IFileConstructorStatistics;
+	class IInstallActionStatistics;
 	class IVerifierStatistics;
 	class IInstallerAnalytics;
 	class IDownloadService;
@@ -186,11 +187,12 @@ namespace BuildPatchServices
 		TSharedPtr<IInstallerAnalytics> InstallerAnalytics;
 
 		// Installer statistics tracking.
+		TUniquePtr<IInstallActionStatistics> InstallActionStatistics;
 		TUniquePtr<IFileOperationTracker> FileOperationTracker;
-		TUniquePtr<ISpeedRecorder> DownloadSpeedRecorder;
-		TUniquePtr<ISpeedRecorder> DiskReadSpeedRecorder;
-		TUniquePtr<ISpeedRecorder> DiskWriteSpeedRecorder;
-		TUniquePtr<IChunkDataSizeProvider> ChunkDataSizeProvider;
+		TSharedPtr<ISpeedRecorder> DownloadSpeedRecorder;
+		TSharedPtr<ISpeedRecorder> DiskReadSpeedRecorder;
+		TSharedPtr<ISpeedRecorder> DiskWriteSpeedRecorder;
+		TSharedPtr<IChunkDataSizeProvider> ChunkDataSizeProvider;
 		TSharedPtr<IDownloadServiceStatistics> DownloadServiceStatistics;
 		TUniquePtr<IChunkDbChunkSourceStatistics> ChunkDbChunkSourceStatistics;
 		TUniquePtr<IInstallChunkSourceStatistics> InstallChunkSourceStatistics;
diff --git a/Engine/Source/Runtime/Online/BuildPatchServices/Private/BuildPatchManifest.cpp b/Engine/Source/Runtime/Online/BuildPatchServices/Private/BuildPatchManifest.cpp
--- a/Engine/Source/Runtime/Online/BuildPatchServices/Private/BuildPatchManifest.cpp
+++ b/Engine/Source/Runtime/Online/BuildPatchServices/Private/BuildPatchManifest.cpp
@@ -13,9 +13,11 @@
 #include "Serialization/JsonReader.h"
 #include "Serialization/JsonSerializer.h"
 #include "Algo/Sort.h"
+#include "Algo/AllOf.h"
 #include "Algo/Accumulate.h"
 #include "Algo/Transform.h"
 #include "Algo/RemoveIf.h"
+#include "Common/FileSystem.h"
 #include "Core/BlockStructure.h"
 #include "Core/AsyncHelpers.h"
 #include "Data/ChunkData.h"
@@ -104,33 +106,6 @@ FString GetFilename(const FString& Filename, const FString& /*RootDirectory*/, b
 	return Filename;
 }
 
-template<typename Cont>
-TMap<FString, int64> GetFileSizes(const Cont& FilesData, const FString& RootDirectory, bool bIsStaged = false)
-{
-	TArray<TTuple<FString, FString, int64>> TargetFiles;
-	TargetFiles.Reserve(FilesData.Num());
-	for (const auto& File : FilesData)
-	{
-		TargetFiles.Add(TTuple<FString, FString, int64>(GetFilename(File, RootDirectory, false), GetFilename(File, RootDirectory, bIsStaged), 0));
-	}
-	auto GetFileSizesTasks = CreateTaskSchedulerLimitedByCoresNumber<int64>(EAsyncExecution::TaskGraph);
-	for (auto& FileInfo : TargetFiles)
-	{
-		GetFileSizesTasks.ScheduleTask(
-			[&RootDirectory, &FileInfo]() { return IFileManager::Get().FileSize(*(RootDirectory / FileInfo.Get<1>())); },
-			[&FileInfo](int64 FileSize) { FileInfo.Get<2>() = FileSize; }
-		);
-	}
-	GetFileSizesTasks.WaitAll();
-	TMap<FString, int64> Results;
-	for (auto& FileInfo : TargetFiles)
-	{
-		Results.FindOrAdd(FileInfo.Get<0>(), FileInfo.Get<2>());
-	}
-	return Results;
-}
-
-
 /* FBuildPatchCustomField implementation
 *****************************************************************************/
 FBuildPatchCustomField::FBuildPatchCustomField(const FString& Value)
@@ -1219,6 +1194,11 @@ const FFileManifest* FBuildPatchAppManifest::GetFileManifest(const FString& File
 	return (FileManifest != nullptr) ? (*FileManifest) : nullptr;
 }
 
+TArray<FFileManifest>::TConstIterator FBuildPatchAppManifest::GetFileManifestIterator() const
+{
+	return FileManifestList.FileList.CreateConstIterator();
+}
+
 bool FBuildPatchAppManifest::IsFileDataManifest() const
 {
 	return ManifestMeta.bIsFileData;
@@ -1336,6 +1316,12 @@ bool FBuildPatchAppManifest::IsManifestEncrypted() const
 	return bIsManifestEncrypted;
 }
 
+bool FBuildPatchAppManifest::IsOriginallyFullyEncrypted() const
+{
+	return EncryptionSecretId.IsValid() 
+		&& Algo::AllOf(ChunkDataList.ChunkList, [](const FChunkInfo& Info) { return Info.EncryptionSecretId.IsValid(); });
+}
+
 TSet<FGuid> FBuildPatchAppManifest::GetNecessaryEncryptionSecretIds() const
 {
 	TSet<FGuid> EncryptionSecretIds;
@@ -1761,7 +1747,7 @@ void FBuildPatchAppManifest::GetOutdatedFiles(const FBuildPatchAppManifest* OldM
 	{
 		// Enumerate files in the this file list, that do not exist, or have different hashes in the OldManifest
 		// to be files no longer required by the build
-		TMap<FString, int64> OutdatedFileSizes = bCheckExistingFile ? GetFileSizes(FilesToCheck, InstallDirectory) : TMap<FString, int64>{};
+		TMap<FString, int64> OutdatedFileSizes = bCheckExistingFile ? ParallelGetFileSizes(FilesToCheck, InstallDirectory) : TMap<FString, int64>{};
 		for (const FString& FileToCheck : FilesToCheck)
 		{
 			const FFileManifest* NewFile = GetFileManifest(FileToCheck);
@@ -1797,7 +1783,7 @@ void FBuildPatchAppManifest::GetMissedFiles(const FString& InstallDirectory, con
 	TSet<FString> Filenames;
 	GetFileList(Filenames);
 	// Enumerate files in the this file list, that do not exist on disk or missed in manifest
-	const TMap<FString, int64> OutdatedFileSizes = GetFileSizes(FilesToCheck.Intersect(Filenames), InstallDirectory);
+	const TMap<FString, int64> OutdatedFileSizes = ParallelGetFileSizes(FilesToCheck.Intersect(Filenames), InstallDirectory);
 	
 	for (const TPair<FString, int64>& Item : OutdatedFileSizes)
 	{
diff --git a/Engine/Source/Runtime/Online/BuildPatchServices/Private/BuildPatchManifest.h b/Engine/Source/Runtime/Online/BuildPatchServices/Private/BuildPatchManifest.h
--- a/Engine/Source/Runtime/Online/BuildPatchServices/Private/BuildPatchManifest.h
+++ b/Engine/Source/Runtime/Online/BuildPatchServices/Private/BuildPatchManifest.h
@@ -105,6 +105,7 @@ public:
 	virtual const FString& GetUniqueBuildId() const override;
 	virtual const FGuid& GetEncryptionSecretId(bool* bOutIsManifestEncrypted = nullptr) const override;
 	virtual bool IsManifestEncrypted() const override;
+	virtual bool IsOriginallyFullyEncrypted() const override;
 	virtual TSet<FGuid> GetNecessaryEncryptionSecretIds() const override;
 	virtual const FString& GetLaunchExe() const override;
 	virtual const FString& GetLaunchCommand() const override;
@@ -258,6 +259,12 @@ public:
 	 */
 	virtual const BuildPatchServices::FFileManifest* GetFileManifest(const FString& Filename) const;
 
+	/**
+	 *  @return a const iterator for the containing file manifests.
+	 *
+	 */
+	virtual TArray<BuildPatchServices::FFileManifest>::TConstIterator GetFileManifestIterator() const;
+
 	/**
 	 * Gets whether this manifest is made up of file data instead of chunk data
 	 * @return	True if the build is made from file data. False if the build is constructed from chunk data.
diff --git a/Engine/Source/Runtime/Online/BuildPatchServices/Private/BuildPatchManifestSet.cpp b/Engine/Source/Runtime/Online/BuildPatchServices/Private/BuildPatchManifestSet.cpp
--- a/Engine/Source/Runtime/Online/BuildPatchServices/Private/BuildPatchManifestSet.cpp
+++ b/Engine/Source/Runtime/Online/BuildPatchServices/Private/BuildPatchManifestSet.cpp
@@ -2,10 +2,12 @@
 
 #include "IBuildManifestSet.h"
 #include "Algo/Accumulate.h"
+#include "Algo/AnyOf.h"
 #include "Algo/Transform.h"
 
 #include "BuildPatchSettings.h"
 
+#include "Common/FileSystem.h"
 #include "Data/ManifestData.h"
 #include "Installer/OptimisedDelta.h"
 #include "BuildPatchUtil.h"
@@ -40,56 +42,65 @@ namespace BuildPatchServices
 		FBuildPatchManifestSet(TArray<FBuildPatchInstallerAction>&& InInstallerActions)
 			: InstallerActions(MoveTemp(InInstallerActions))
 		{
-			// We need to perform our own lookups to avoid looping in some scenarios.
-			for (const FBuildPatchInstallerAction& InstallerAction : InstallerActions)
+			// We need to perform our own lookups to avoid looping in some scenarios, and also allow other systems to not have to understand subdirectories.
+			// All publically used API includes subdirectory in filenames.
+			// Reverse iterate provided actions, making sure there is only one unique reference to each file in the entire action set.
+			// Though an undesirable request, the expected behaviour for multiple actions producing the same file would be clobbering by later actions.
+			for (int32 Idx = InstallerActions.Num() - 1; Idx >= 0; --Idx)
 			{
+				FBuildPatchInstallerAction& InstallerAction = InstallerActions[Idx];
+				// We can add every DataLookup from all manifests, since any file could reference any data piece.
+				// Duplicates are not an issue because any single ChunkGuid, will resolve to a reachable Chunk file.
 				for (const FBuildPatchAppManifest& Manifest : InstallerAction)
 				{
 					for (const TPair<FGuid, const BuildPatchServices::FChunkInfo*>& ChunkInfoPair : Manifest.ChunkInfoLookup)
 					{
-						// Order is crucial here. The ManifestIterator above starts with the current manifest and then does the install manifest,
-						// so the data lookup prefers the new chunks over the deployed chunks.
 						DataLookup.Add(ChunkInfoPair.Key, FDataReference{ ChunkInfoPair.Value, &Manifest, &InstallerAction });
 					}
 				}
-				if (!InstallerAction.IsInstall())
+
+				// Get all the CurrentFiles, if there are any.
+				const FBuildPatchAppManifest* CurrentManifest = InstallerAction.TryGetCurrentManifest();
+				if (CurrentManifest != nullptr)
 				{
-					for (const TPair<FStringView, const BuildPatchServices::FFileManifest*>& FileManifestPair : InstallerAction.GetCurrentManifest().FileManifestLookup)
+					for (const FFileManifest* CurrFileManifest : CurrentManifest->GetTaggedFileManifests(InstallerAction.GetInstallTags()))
 					{
-						CurrentFileLookup.Add(InstallerAction.GetInstallSubdirectory() / FString(FileManifestPair.Key), FFileReference{ FileManifestPair.Value, InstallerAction.TryGetCurrentManifest(), &InstallerAction });
+						FString FullFilename = InstallerAction.GetInstallSubdirectory() / CurrFileManifest->Filename;
+						// Do not allow clobbering.
+						if (!CurrentFileLookup.Contains(FullFilename))
+						{
+							CurrentFileLookup.Add(MoveTemp(FullFilename), FFileReference{CurrFileManifest, CurrentManifest, &InstallerAction});
+						}
 					}
 				}
-			}
-			// Reverse iterate provided actions, making sure there is only one unique reference to each file.
-			// Though an undesirable request, the expected behaviour for multiple actions producing the same file would be clobbering by later actions.
-			TSet<FString> VisitedFiles;
-			for (int32 Idx = InstallerActions.Num() - 1; Idx >= 0; --Idx)
-			{
-				FBuildPatchInstallerAction& InstallerAction = InstallerActions[Idx];
-				const bool bIsUninstall = InstallerAction.IsUninstall();
-				TSet<FString> FilesToFilter;
-				// Ensure that all tagged files, have their new file manifest in the lookup, and current equiv in the lookup in case of multiple file references.
-				for (const FString& TaggedFile : InstallerAction.GetTaggedFiles())
+
+				// Get all the NewFiles, if there are any.
+				const FBuildPatchAppManifest* InstallManifest = InstallerAction.TryGetInstallManifest();
+				if (InstallManifest != nullptr)
 				{
-					FString FullFilename = InstallerAction.GetInstallSubdirectory() / TaggedFile;
-					if (!bIsUninstall)
+					// Track files that were actually placed into the lookups, so we can override the action's tagged file list.
+					// This will help other functionality guarentee it cannot return the incorrect resolved file manifest, sice each file will only be referenced
+					// by one action after this logic.
+					// We are also going to iterate files in order of the manifest, rather than just the generic public functions to get tagged files, as this will allow us to preserve the desired manifest
+					// installation order much easier, since we can treat the tagged file list as sorted because it is construct once.
+					// REMEMBER: TaggedFiles refer to the manifest's filenames, so not including subdirectory. If two actions share a filename, but with different subdirectories, they both get installed.
+					//           This is important because the tagged file list is used to iterate file manifest in the build manifest.
+					TSet<FString> ResultingTaggedFiles;
+					for (auto It = InstallManifest->GetFileManifestIterator(); It; ++It)
 					{
-						bool bDuplicateFilename = false;
-						VisitedFiles.Add(FullFilename, &bDuplicateFilename);
-						if (bDuplicateFilename)
+						const bool bFileIsTagged = It->InstallTags.Num() == 0 || Algo::AnyOf(It->InstallTags, [&](const FString& Tag) { return InstallerAction.GetInstallTags().Contains(Tag); });
+						if (bFileIsTagged)
 						{
-							FilesToFilter.Add(TaggedFile);
+							FString FullFilename = InstallerAction.GetInstallSubdirectory() / It->Filename;
+							// Do not allow clobbering.
+							if (!NewFileLookup.Contains(FullFilename))
+							{
+								NewFileLookup.Add(MoveTemp(FullFilename), FFileReference{ &(*It), InstallManifest, &InstallerAction });
+								ResultingTaggedFiles.Add(It->Filename);
+							}
 						}
 					}
-					if (!InstallerAction.IsInstall() && InstallerAction.GetCurrentManifest().FileManifestLookup.Contains(TaggedFile))
-					{
-						CurrentFileLookup.Add(FullFilename, FFileReference{ InstallerAction.GetCurrentManifest().FileManifestLookup[TaggedFile], InstallerAction.TryGetCurrentManifest(), &InstallerAction });
-					}
-					NewFileLookup.Add(MoveTemp(FullFilename), FFileReference{ InstallerAction.GetInstallManifest().FileManifestLookup[TaggedFile], InstallerAction.TryGetInstallManifest(), &InstallerAction });
-				}
-				if (FilesToFilter.Num() > 0)
-				{
-					InstallerAction.SetTaggedFiles(InstallerAction.GetTaggedFiles().Difference(FilesToFilter));
+					InstallerAction.SetTaggedFiles(MoveTemp(ResultingTaggedFiles));
 				}
 			}
 		}
@@ -145,21 +156,18 @@ namespace BuildPatchServices
 
 		virtual void GetReferencedChunks(TArray<FGuid>& DataGuids) const override
 		{
-			// For an array instead of set, we would expect to provide duplicates
+			// For an array instead of set, we would expect to provide ordered duplicates.
 			for (const FBuildPatchInstallerAction& InstallerAction : InstallerActions)
 			{
-				if (!InstallerAction.IsUninstall())
+				// InstallerAction.GetTaggedFiles() has already been de-duped across actions in this class's constructor.
+				for (const FString& TaggedFile : InstallerAction.GetTaggedFiles())
 				{
-					for (const FString& TaggedFile : InstallerAction.GetTaggedFiles())
+					FFileManifest const * const TaggedFileManifest = GetNewFileManifest(InstallerAction.GetInstallSubdirectory() / TaggedFile);
+					if (TaggedFileManifest)
 					{
-						FFileReference const* const LookupResult = NewFileLookup.Find(TaggedFile);
-						if (LookupResult)
+						for (const FChunkPart& ChunkPart : TaggedFileManifest->ChunkParts)
 						{
-							FFileManifest const* const TaggedFileManifest = LookupResult->Get<0>();
-							for (const FChunkPart& ChunkPart : TaggedFileManifest->ChunkParts)
-							{
-								DataGuids.Add(ChunkPart.Guid);
-							}
+							DataGuids.Add(ChunkPart.Guid);
 						}
 					}
 				}
@@ -201,16 +209,6 @@ namespace BuildPatchServices
 			return false;
 		}
 
-		virtual FString GetDataFilename(const FGuid& DataGuid) const override
-		{
-			FDataReference const* const LookupResult = DataLookupFind(DataGuid);
-			if (LookupResult)
-			{
-				return LookupResult->Get<2>()->GetCloudSubdirectory() / FBuildPatchUtils::GetDataFilename(*LookupResult->Get<1>(), DataGuid);
-			}
-			return FString();
-		}
-
 		virtual FGuid GetDataSecretId(const FGuid& DataGuid) const override
 		{
 			FDataReference const* const LookupResult = DataLookupFind(DataGuid);
@@ -241,6 +239,16 @@ namespace BuildPatchServices
 			return FString();
 		}
 
+		virtual FString GetChunkFilenameWithCloudSubdirectory(const FGuid& DataGuid) const override
+		{
+			FDataReference const* const LookupResult = DataLookupFind(DataGuid);
+			if (LookupResult)
+			{
+				return LookupResult->Get<2>()->GetCloudSubdirectory() / FBuildPatchUtils::GetDataFilename(*LookupResult->Get<1>(), DataGuid);
+			}
+			return FString();
+		}
+
 		virtual FString GetChunkUniqueBuildId(const FGuid& DataGuid) const override
 		{
 			FDataReference const* const LookupResult = DataLookupFind(DataGuid);
@@ -277,14 +285,57 @@ namespace BuildPatchServices
 
 		virtual void GetOutdatedFiles(const FString& InstallDirectory, TSet<FString>& OutdatedFiles) const override
 		{
-			TSet<FString> TempFileSet;
+			// For outdated files, we must consider the difference by comparing the file manifests
+			// according to CurrentFileLookup and NewFileLookup, since these have already been properly flattened
+			// using the correct action resolution logic.
+
+			// For this reason, we cannot use the GetOutdatedFiles function on the actual build manifest,
+			// since it is not subdirectory aware, or multi-action aware.
+
+			const bool bCheckFileSizeOnDisk = InstallDirectory.IsEmpty() == false;
+			TSet<FString> FileSizesToCheck;
 			for (const FBuildPatchInstallerAction& InstallerAction : InstallerActions)
 			{
-				if (!InstallerAction.IsUninstall())
+				// InstallerAction.GetTaggedFiles() has already been de-duped across actions in this class's constructor.
+				for (const FString& ManifestFilename : InstallerAction.GetTaggedFiles())
 				{
-					TempFileSet.Reset();
-					InstallerAction.GetInstallManifest().GetOutdatedFiles(InstallerAction.TryGetCurrentManifest(), InstallDirectory / InstallerAction.GetInstallSubdirectory(), InstallerAction.GetTaggedFiles(), TempFileSet);
-					Algo::Transform(TempFileSet, OutdatedFiles, [&](const FString& Filename) { return InstallerAction.GetInstallSubdirectory() / Filename; });
+					FString FullFilename = InstallerAction.GetInstallSubdirectory() / ManifestFilename;
+					const FFileManifest* CurrentFileManifest = GetCurrentFileManifest(FullFilename);
+					const FFileManifest* NewFileManifest = GetNewFileManifest(FullFilename);
+
+					// Automatically considered outdated, if we have a new manifest without current manfiest.
+					if (CurrentFileManifest == nullptr && NewFileManifest != nullptr)
+					{
+						OutdatedFiles.Add(MoveTemp(FullFilename));
+					}
+					// If we have both, check if the hash changed.
+					else if (CurrentFileManifest != nullptr && NewFileManifest != nullptr
+					     && (CurrentFileManifest->SHA1Hash != NewFileManifest->SHA1Hash))
+					{
+						OutdatedFiles.Add(MoveTemp(FullFilename));
+					}
+					// If we are checking disk sizes, add this 'not outdated' file for the check.
+					else if (bCheckFileSizeOnDisk)
+					{
+						FileSizesToCheck.Add(MoveTemp(FullFilename));
+					}
+				}
+			}
+			// Run the disk size check for any files that need it.
+			if (bCheckFileSizeOnDisk && FileSizesToCheck.Num() > 0)
+			{
+				TMap<FString, int64> FileSizes = ParallelGetFileSizes(FileSizesToCheck, InstallDirectory);
+				for (FString& FullFilename : FileSizesToCheck)
+				{
+					const FFileManifest* NewFileManifest = GetNewFileManifest(FullFilename);
+					if (NewFileManifest != nullptr)
+					{
+						const int64 ActualFileSize = FileSizes.FindRef(FullFilename, INDEX_NONE);
+						if (NewFileManifest->FileSize != ActualFileSize)
+						{
+							OutdatedFiles.Add(MoveTemp(FullFilename));
+						}
+					}
 				}
 			}
 		}
diff --git a/Engine/Source/Runtime/Online/BuildPatchServices/Private/BuildPatchServicesModule.cpp b/Engine/Source/Runtime/Online/BuildPatchServices/Private/BuildPatchServicesModule.cpp
--- a/Engine/Source/Runtime/Online/BuildPatchServices/Private/BuildPatchServicesModule.cpp
+++ b/Engine/Source/Runtime/Online/BuildPatchServices/Private/BuildPatchServicesModule.cpp
@@ -43,10 +43,10 @@ IMPLEMENT_MODULE(FBuildPatchServicesModule, BuildPatchServices);
 void FBuildPatchServicesModule::StartupModule()
 {
 	// Debug sanity checks
+	static_assert(UE_ARRAY_COUNT(InstallErrorPrefixes::ErrorTypeStrings) == (uint64)EBuildPatchInstallError::InvalidOrMax, "Please add support for the extra values to the InstallErrorPrefixes::ErrorTypeStrings.");
 #if UE_BUILD_DEBUG
 	TSet<FString> NoDupes;
 	bool bWasDupe = false;
-	check(UE_ARRAY_COUNT(InstallErrorPrefixes::ErrorTypeStrings) == (uint64)EBuildPatchInstallError::NumInstallErrors);
 	for (int32 Idx = 0; Idx < UE_ARRAY_COUNT(InstallErrorPrefixes::ErrorTypeStrings); ++Idx)
 	{
 		NoDupes.Add(FString(InstallErrorPrefixes::ErrorTypeStrings[Idx]), &bWasDupe);
diff --git a/Engine/Source/Runtime/Online/BuildPatchServices/Private/Common/FileSystem.cpp b/Engine/Source/Runtime/Online/BuildPatchServices/Private/Common/FileSystem.cpp
--- a/Engine/Source/Runtime/Online/BuildPatchServices/Private/Common/FileSystem.cpp
+++ b/Engine/Source/Runtime/Online/BuildPatchServices/Private/Common/FileSystem.cpp
@@ -9,8 +9,10 @@
 #include "HAL/FileManager.h"
 #include "HAL/PlatformFileManager.h"
 #include "HAL/FileManagerGeneric.h"
-#include "Misc/OutputDeviceRedirector.h"
 #include "Misc/FileHelper.h"
+#include "Misc/OutputDeviceRedirector.h"
+
+#include "Core/AsyncHelpers.h"
 
 #if PLATFORM_MAC
 #include <sys/stat.h>
@@ -192,6 +194,32 @@ static_assert((uint32)BuildPatchServices::EReadFlags::AllowWrite == (uint32)::EF
 
 namespace BuildPatchServices
 {
+	TMap<FString, int64> ParallelGetFileSizes(const TSet<FString>& FilesToCheck, const FString& RootDirectory)
+	{
+		TArray<TTuple<FString, int64>> TargetFiles;
+		TargetFiles.Reserve(FilesToCheck.Num());
+		for (const FString& File : FilesToCheck)
+		{
+			TargetFiles.Add(TTuple<FString, int64>(File, INDEX_NONE));
+		}
+		auto GetFileSizesTasks = CreateTaskSchedulerLimitedByCoresNumber<int64>(EAsyncExecution::TaskGraph);
+		for (TTuple<FString, int64>& FileInfo : TargetFiles)
+		{
+			GetFileSizesTasks.ScheduleTask(
+				[&RootDirectory, &FileInfo]() { return IFileManager::Get().FileSize(*(RootDirectory / FileInfo.Get<0>())); },
+				[&FileInfo](int64 FileSize) { FileInfo.Get<1>() = FileSize; }
+			);
+		}
+		GetFileSizesTasks.WaitAll();
+		TMap<FString, int64> Results;
+		for (TTuple<FString, int64>& FileInfo : TargetFiles)
+		{
+			Results.Add(FileInfo.Get<0>(), FileInfo.Get<1>());
+		}
+		return Results;
+	}
+
+
 	class FDirectoryEnumeratorBase : public IPlatformFile::FDirectoryVisitor
 	{
 	protected:
diff --git a/Engine/Source/Runtime/Online/BuildPatchServices/Private/Common/FileSystem.h b/Engine/Source/Runtime/Online/BuildPatchServices/Private/Common/FileSystem.h
--- a/Engine/Source/Runtime/Online/BuildPatchServices/Private/Common/FileSystem.h
+++ b/Engine/Source/Runtime/Online/BuildPatchServices/Private/Common/FileSystem.h
@@ -43,6 +43,14 @@ namespace BuildPatchServices
 	};
 	ENUM_CLASS_FLAGS(EReadFlags);
 
+	/**
+	 * A static helper function to parallelise retrieving file sizes for performance on large sets of files.
+	 * @param FilesToCheck    The set of files to lookup the filesize for
+	 * @param RootDirectory   The root directory of the files. Empty string can be passed if FilesToCheck has full file paths.
+	 * @return a map of string from FilesToCheck to size on disk. INDEX_NONE if the file does not exist.
+	 */
+	TMap<FString, int64> ParallelGetFileSizes(const TSet<FString>& FilesToCheck, const FString& RootDirectory);
+
 	/**
 	 * The File System class is used for classes which require file access. It wraps Core IFileManager, and IPlatformFile. Also provides additional
 	 * functionality missing from these classes at the time of writing.
@@ -276,4 +284,4 @@ namespace BuildPatchServices
 		 */
 		static IFileSystem* Create();
 	};
-}
+}
\ No newline at end of file
diff --git a/Engine/Source/Runtime/Online/BuildPatchServices/Private/Data/ChunkData.cpp b/Engine/Source/Runtime/Online/BuildPatchServices/Private/Data/ChunkData.cpp
--- a/Engine/Source/Runtime/Online/BuildPatchServices/Private/Data/ChunkData.cpp
+++ b/Engine/Source/Runtime/Online/BuildPatchServices/Private/Data/ChunkData.cpp
@@ -12,8 +12,8 @@
 #include "Common/StatsCollector.h"
 #include "BuildPatchHash.h"
 #include "BuildPatchManifest.h"
-#include "Installer/CloudChunkSource.h"
 #include "BuildPatchServicesPrivate.h"
+#include "Installer/CloudChunkSource.h"
 
 // The chunk header magic codeword, for quick checking that the opened file is a chunk file.
 #define CHUNK_HEADER_MAGIC      0xB1FE3AA2
@@ -678,10 +678,29 @@ namespace BuildPatchServices
 				// Encrypted but not compressed, copy before checking hashes
 				FMemory::Memcpy(UseDestinationBuffer.GetData(), UseCompressedBuffer.GetData(), FMath::Min(UseDestinationBuffer.GetSize(), UseCompressedBuffer.GetSize()));
 			}
+			
+			// Uncompressed - where ValidateAndRead put the data depends on whether the destination buffer
+			// was big enough for the full chunk.
+			FMemoryView HashCheckBuffer;
+			if ((InHeader.StoredAs == EChunkStorageFlags::RawData) &&
+			    (InHeader.DataSizeUncompressed != InDestinationBuffer.GetSize()))
+			{
+				HashCheckBuffer = InCompressedBuffer.GetView();
+			}
+			else
+			{
+				HashCheckBuffer = UseDestinationBuffer;
+			}
 
-			bool bHashMatches = CheckChunkHash(InHeader, UseDestinationBuffer.GetData(), UseDestinationBuffer.GetSize());
+			bool bHashMatches = CheckChunkHash(InHeader, HashCheckBuffer.GetData(), HashCheckBuffer.GetSize());
 			if (bHashMatches && InHeader.DataSizeUncompressed != InDestinationBuffer.GetSize())
 			{
+				// Destination buffer should always be <= the chunk size.
+				if (InDestinationBuffer.GetSize() > TempDestinationBuffer.GetSize())
+				{
+					UE_LOG(LogChunkData, Error, TEXT("Invalid destination buffer size for chunk %s"), *WriteToString<64>(InHeader.Guid));
+					return false;
+				}
 				// We routed through a temp buffer, copy out as necessary.
 				FMemory::Memcpy(InDestinationBuffer.GetData(), TempDestinationBuffer.GetData(), InDestinationBuffer.GetSize());
 			}
diff --git a/Engine/Source/Runtime/Online/BuildPatchServices/Private/Data/ManifestData.cpp b/Engine/Source/Runtime/Online/BuildPatchServices/Private/Data/ManifestData.cpp
--- a/Engine/Source/Runtime/Online/BuildPatchServices/Private/Data/ManifestData.cpp
+++ b/Engine/Source/Runtime/Online/BuildPatchServices/Private/Data/ManifestData.cpp
@@ -918,6 +918,15 @@ namespace BuildPatchServices
 		{
 			return false;
 		}
+
+		// if we are trying to save a format we don't recognize, return with error
+		if (Ar.IsSaving() && OriginalFormat > EFeatureLevel::Latest)
+		{
+			UE_LOG(LogManifestData, Error, TEXT("FManifestData::Serialize: Trying to save an unrecognized format, Latest format %d, Original format %d"),
+				static_cast<uint32>(EFeatureLevel::Latest), static_cast<uint32>(OriginalFormat));
+			Ar.SetError();
+			return false;
+		}
 		bool bSuccess = true;
 		// If we are saving an old format, defer to the old code!
 		if (Ar.IsSaving() && SaveFormat < EFeatureLevel::StoredAsBinaryData)
diff --git a/Engine/Source/Runtime/Online/BuildPatchServices/Private/Diffing/DiffManifests.cpp b/Engine/Source/Runtime/Online/BuildPatchServices/Private/Diffing/DiffManifests.cpp
--- a/Engine/Source/Runtime/Online/BuildPatchServices/Private/Diffing/DiffManifests.cpp
+++ b/Engine/Source/Runtime/Online/BuildPatchServices/Private/Diffing/DiffManifests.cpp
@@ -299,15 +299,15 @@ namespace BuildPatchServices
 		bool AsyncRun();
 		void HandleDownloadComplete(int32 RequestId, const FDownloadRef& Download);
 
-		void ExportPatchDescriptorFiles(const TSet<FString>& InstallTags, FBuildPatchAppManifestPtr OldManifest, FBuildPatchAppManifestPtr NewManifest);
+		void ExportPatchDescriptorFiles(const TSet<FString>& InstallTags, const FBuildPatchAppManifestPtr &OldManifest, const FBuildPatchAppManifestPtr &NewManifest);
 	private:
 		const FDiffManifestsConfiguration Configuration;
 		FDownloadCompleteDelegate DownloadCompleteDelegate;
 		FDownloadProgressDelegate DownloadProgressDelegate;
 		TSharedPtr<IFileSystem> FileSystem;
 		TSharedPtr<IHttpManager> HttpManager;
-		TUniquePtr<IChunkDataSizeProvider> ChunkDataSizeProvider;
-		TUniquePtr<ISpeedRecorder> DownloadSpeedRecorder;
+		TSharedPtr<IChunkDataSizeProvider> ChunkDataSizeProvider;
+		TSharedPtr<ISpeedRecorder> DownloadSpeedRecorder;
 		TSharedPtr<IInstallerAnalytics> InstallerAnalytics;
 		TSharedPtr<IDownloadServiceStatistics> DownloadServiceStatistics;
 		TUniquePtr<IDownloadService> DownloadService;
@@ -336,7 +336,7 @@ namespace BuildPatchServices
 		, ChunkDataSizeProvider(FChunkDataSizeProviderFactory::Create({}))
 		, DownloadSpeedRecorder(FSpeedRecorderFactory::Create())
 		, InstallerAnalytics(FInstallerAnalyticsFactory::Create(nullptr))
-		, DownloadServiceStatistics(FDownloadServiceStatisticsFactory::Create(DownloadSpeedRecorder.Get(), ChunkDataSizeProvider.Get(), InstallerAnalytics.Get()))
+		, DownloadServiceStatistics(FDownloadServiceStatisticsFactory::Create(DownloadSpeedRecorder, ChunkDataSizeProvider, InstallerAnalytics))
 		, DownloadService(FDownloadServiceFactory::Create(HttpManager, FileSystem, DownloadServiceStatistics, InstallerAnalytics))
 		, StatsCollector(FStatsCollectorFactory::Create())
 		, bShouldRun(true)
@@ -420,7 +420,24 @@ namespace BuildPatchServices
 		return Thread.Get();
 	}
 
-	void FDiffManifests::ExportPatchDescriptorFiles(const TSet<FString>& InstallTags, FBuildPatchAppManifestPtr OldManifest, FBuildPatchAppManifestPtr NewManifest)
+	static void WriteStringsToFile(const TUtf8StringBuilder<1024> &Listing, const FString &FileName)
+	{
+		TUniquePtr<FArchive> Ar(IFileManager::Get().CreateFileWriter(*FileName, 0));
+		if (!Ar)
+		{
+			UE_LOG(LogTemp, Warning, TEXT("Unable to open output file: %s"), *FileName);
+		}
+		else
+		{
+			FUtf8StringView OutputView = MakeStringView(Listing);
+			UTF8CHAR UTF8BOM[] = { (UTF8CHAR)0xEF, (UTF8CHAR)0xBB, (UTF8CHAR)0xBF };
+			Ar->Serialize(&UTF8BOM, sizeof(UTF8BOM));
+			Ar->Serialize((void*)OutputView.GetData(), OutputView.Len() * sizeof(UTF8CHAR));
+			Ar->Close();
+		}
+	}
+
+	void FDiffManifests::ExportPatchDescriptorFiles(const TSet<FString>& InstallTags, const FBuildPatchAppManifestPtr &OldManifest, const FBuildPatchAppManifestPtr &NewManifest)
 	{
 		UE_LOG(LogDiffManifests, Display, TEXT("Generating PatchDescription Files for %d tags and writing to %s"), InstallTags.Num(), *Configuration.OutputPatchDescriptorPath);
 
@@ -443,7 +460,6 @@ namespace BuildPatchServices
 			TSet<FString> FilesToPatch;
 			for (const FString& NewFile : NewFiles)
 			{
-				const FFileManifest* FileManifest = NewManifest->GetFileManifest(NewFile);
 				if (!OldFiles.Contains(NewFile))
 				{
 					FilesToPatch.Add(NewFile);
@@ -478,11 +494,11 @@ namespace BuildPatchServices
 			OldManifest->GetChunksRequiredForFiles(OldFiles, ChunksForOldFiles);
 			OldManifest->EnumerateProducibleChunks(CurrentInstallTag, ChunksForOldFiles, ChunksOnDisk);
 
-			TSet<FGuid> ChunksToDownload = ChunksForNewFiles.Difference(ChunksOnDisk);
+			const TSet<FGuid> ChunksToDownload = ChunksForNewFiles.Difference(ChunksOnDisk);
 			uint64 TotalNewChunkBytes = 0;
 			uint64 TotalNewChunkDownloadBytes = 0;
 
-			for (FGuid& Guid : ChunksToDownload)
+			for (const FGuid& Guid : ChunksToDownload)
 			{
 				const BuildPatchServices::FChunkInfo* ChunkInfo = NewManifest->GetChunkInfo(Guid);
 				TotalNewChunkBytes += ChunkInfo->DataSizeUncompressed;
@@ -560,22 +576,9 @@ namespace BuildPatchServices
 					}
 				}
 
-				FString FileName = (Configuration.OutputPatchDescriptorPath / PatchFile) + ".patch";
+				const FString FileName = (Configuration.OutputPatchDescriptorPath / PatchFile) + ".patch";
 
-				TUniquePtr<FArchive> Ar(IFileManager::Get().CreateFileWriter(*FileName, 0));
-				if (!Ar)
-				{
-					UE_LOG(LogTemp, Warning, TEXT("Unable to open output file: %s"), *FileName);
-				}
-				else
-				{
-					FUtf8StringView OutputView = MakeStringView(Listing);
-					UTF8CHAR UTF8BOM[] = { (UTF8CHAR)0xEF, (UTF8CHAR)0xBB, (UTF8CHAR)0xBF };
-					Ar->Serialize(&UTF8BOM, sizeof(UTF8BOM));
-					Ar->Serialize((void*)OutputView.GetData(), OutputView.Len() * sizeof(UTF8CHAR));
-					Ar->Close();
-				}
-						
+				WriteStringsToFile(Listing, FileName);
 			} // end each file
 
 			// This can be more than the uncompressed download size if chunks are reused, or smaller if we didn't need
diff --git a/Engine/Source/Runtime/Online/BuildPatchServices/Private/Generation/ChunkDeltaOptimiser.cpp b/Engine/Source/Runtime/Online/BuildPatchServices/Private/Generation/ChunkDeltaOptimiser.cpp
--- a/Engine/Source/Runtime/Online/BuildPatchServices/Private/Generation/ChunkDeltaOptimiser.cpp
+++ b/Engine/Source/Runtime/Online/BuildPatchServices/Private/Generation/ChunkDeltaOptimiser.cpp
@@ -181,7 +181,7 @@ namespace DeltaOptimiseHelpers
 		});
 	}
 
-	void MakeScannerLocalList(const TMap<FString, FString>& AliasedTags, FChunkSearcher& ChunkSearcher, IDeltaChunkEnumeration* Enumeration, const FBlockStructure& BuildStructure, FScannerFilesList& Result)
+	void MakeScannerLocalList(const TMap<FString, FString>& AliasedTags, const TSet<FString>& IgnoredTags, FChunkSearcher& ChunkSearcher, IDeltaChunkEnumeration* Enumeration, const FBlockStructure& BuildStructure, FScannerFilesList& Result)
 	{
 		uint64 FirstByte = 0;
 		ChunkSearcher.ForEachOverlap(BuildStructure, [&](const FBlockRange& OverlapRange, FChunkSearcher::FFileDListNode* File, FChunkSearcher::FChunkDListNode* Chunk)
@@ -195,6 +195,11 @@ namespace DeltaOptimiseHelpers
 			FileTagset.Reserve(File->GetValue().Manifest->InstallTags.Num());
 			for (const FString& InstallTag : File->GetValue().Manifest->InstallTags)
 			{
+				if (IgnoredTags.Contains(InstallTag))
+				{
+					continue;
+				}
+
 				const FString* AliasedTag = AliasedTags.Find(InstallTag);
 				FileTagset.Add(AliasedTag ? *AliasedTag : InstallTag);
 			}
@@ -660,8 +665,8 @@ namespace BuildPatchServices
 		TSharedPtr<IFileSystem> FileSystem;
 		TUniquePtr<ICrypto> Crypto;
 		TSharedPtr<IHttpManager> HttpManager;
-		TUniquePtr<IChunkDataSizeProvider> ChunkDataSizeProvider;
-		TUniquePtr<ISpeedRecorder> DownloadSpeedRecorder;
+		TSharedPtr<IChunkDataSizeProvider> ChunkDataSizeProvider;
+		TSharedPtr<ISpeedRecorder> DownloadSpeedRecorder;
 		TSharedPtr<IInstallerAnalytics> InstallerAnalytics;
 		TSharedPtr<IDownloadServiceStatistics> DownloadServiceStatistics;
 		TUniquePtr<IDownloadService> DownloadService;
@@ -686,7 +691,7 @@ namespace BuildPatchServices
 		, ChunkDataSizeProvider(FChunkDataSizeProviderFactory::Create({}))
 		, DownloadSpeedRecorder(FSpeedRecorderFactory::Create())
 		, InstallerAnalytics(FInstallerAnalyticsFactory::Create(nullptr))
-		, DownloadServiceStatistics(FDownloadServiceStatisticsFactory::Create(DownloadSpeedRecorder.Get(), ChunkDataSizeProvider.Get(), InstallerAnalytics.Get()))
+		, DownloadServiceStatistics(FDownloadServiceStatisticsFactory::Create(DownloadSpeedRecorder, ChunkDataSizeProvider, InstallerAnalytics))
 		, DownloadService(FDownloadServiceFactory::Create(HttpManager, FileSystem, DownloadServiceStatistics, InstallerAnalytics))
 		, MessagePumpA(FMessagePumpFactory::Create())
 		, MessagePumpB(FMessagePumpFactory::Create())
@@ -997,7 +1002,7 @@ namespace BuildPatchServices
 				TUniquePtr<IManifestBuildStreamer> ManifestAStream(FBuildStreamerFactory::Create(MoveTemp(ManifestAStreamConfig), MoveTemp(ManifestAStreamDependencies)));
 
 				// First we re-chunk prev unknown build parts into the scanner window size.
-				TUniquePtr<IDeltaChunkEnumeration> DeltaChunkEnumeration(FDeltaChunkEnumerationFactory::Create(ManifestAStream.Get(), StatsCollector.Get(), *ManifestA.Get(), Configuration.ScanWindowSize, Configuration.ManifestATagAliases));
+				TUniquePtr<IDeltaChunkEnumeration> DeltaChunkEnumeration(FDeltaChunkEnumerationFactory::Create(ManifestAStream.Get(), StatsCollector.Get(), *ManifestA.Get(), Configuration.ScanWindowSize, Configuration.ManifestATagAliases, Configuration.ManifestAIgnoreTags));
 				ChunkingTimer.Start();
 				DeltaChunkEnumeration->Run();
 				ChunkingTimer.Stop();
@@ -1114,7 +1119,7 @@ namespace BuildPatchServices
 
 						FBlockStructure ScannerBuildStructure;
 						ManifestBDesiredBytes.SelectSerialBytes(NewScanner->Offset, SizeToScan, ScannerBuildStructure);
-						DeltaOptimiseHelpers::MakeScannerLocalList(Configuration.ManifestBTagAliases, FileListSearcher, DeltaChunkEnumeration.Get(), ScannerBuildStructure, NewScanner->FilesList);
+						DeltaOptimiseHelpers::MakeScannerLocalList(Configuration.ManifestBTagAliases, Configuration.ManifestBIgnoreTags, FileListSearcher, DeltaChunkEnumeration.Get(), ScannerBuildStructure, NewScanner->FilesList);
 
 						NewScanner->Scanner.Reset(FDeltaScannerFactory::Create(Configuration.ScanWindowSize, NewScanner->Data, NewScanner->FilesList, DeltaChunkEnumeration.Get(), StatsCollector.Get()));
 						ConsumedBufferData += SizeToScan;
@@ -1152,7 +1157,7 @@ namespace BuildPatchServices
 
 								FBlockStructure ScannerBuildStructure;
 								ManifestBDesiredBytes.SelectSerialBytes(NewScanner->Offset, RightFork.GetSize(), ScannerBuildStructure);
-								DeltaOptimiseHelpers::MakeScannerLocalList(Configuration.ManifestBTagAliases, FileListSearcher, DeltaChunkEnumeration.Get(), ScannerBuildStructure, NewScanner->FilesList);
+								DeltaOptimiseHelpers::MakeScannerLocalList(Configuration.ManifestBTagAliases, Configuration.ManifestBIgnoreTags, FileListSearcher, DeltaChunkEnumeration.Get(), ScannerBuildStructure, NewScanner->FilesList);
 
 								NewScanner->Scanner.Reset(FDeltaScannerFactory::Create(Configuration.ScanWindowSize, NewScanner->Data, NewScanner->FilesList, DeltaChunkEnumeration.Get(), StatsCollector.Get()));
 								NewScanner->bIsFinalScanner = DataScannerEntry.bIsFinalScanner;
@@ -1167,7 +1172,7 @@ namespace BuildPatchServices
 
 								ScannerBuildStructure.Empty();
 								ManifestBDesiredBytes.SelectSerialBytes(NewScanner->Offset, LeftFork.GetSize(), ScannerBuildStructure);
-								DeltaOptimiseHelpers::MakeScannerLocalList(Configuration.ManifestBTagAliases, FileListSearcher, DeltaChunkEnumeration.Get(), ScannerBuildStructure, NewScanner->FilesList);
+								DeltaOptimiseHelpers::MakeScannerLocalList(Configuration.ManifestBTagAliases, Configuration.ManifestBIgnoreTags, FileListSearcher, DeltaChunkEnumeration.Get(), ScannerBuildStructure, NewScanner->FilesList);
 
 								NewScanner->Scanner.Reset(FDeltaScannerFactory::Create(Configuration.ScanWindowSize, NewScanner->Data, NewScanner->FilesList, DeltaChunkEnumeration.Get(), StatsCollector.Get()));
 								NewScanner->bIsFinalScanner = false;
diff --git a/Engine/Source/Runtime/Online/BuildPatchServices/Private/Generation/CloudEnumeration.cpp b/Engine/Source/Runtime/Online/BuildPatchServices/Private/Generation/CloudEnumeration.cpp
--- a/Engine/Source/Runtime/Online/BuildPatchServices/Private/Generation/CloudEnumeration.cpp
+++ b/Engine/Source/Runtime/Online/BuildPatchServices/Private/Generation/CloudEnumeration.cpp
@@ -38,10 +38,11 @@ namespace BuildPatchServices
 		 *
 		 * @param CloudDirectory The cloud root directory.
 		 * @param ManifestAgeThreshold The threshold age for manifest files; files older than this will be ignored.
+		 * @param Predicate A callable that takes a const reference to an IBuildManifest and returns true to include the manifest, or false to ignore it.
 		 * @param OutputFeatureLevel The desired feature level for enumerated output.
 		 * @param StatsCollector A pointer to the statistics collector for recording related data.
 		 */
-		FCloudEnumeration(const FString& CloudDirectory, const FDateTime& ManifestAgeThreshold, const EFeatureLevel& InOutputFeatureLevel, FStatsCollector* StatsCollector);
+		FCloudEnumeration(const FString& CloudDirectory, const FDateTime& ManifestAgeThreshold, const TFunction<bool(const IBuildManifest&)>& InPredicate, const EFeatureLevel& InOutputFeatureLevel, FStatsCollector* StatsCollector);
 		virtual ~FCloudEnumeration();
 
 		virtual bool IsComplete() const override;
@@ -68,6 +69,7 @@ namespace BuildPatchServices
 	private:
 		const FString CloudDirectory;
 		const FDateTime ManifestAgeThreshold;
+		const TFunction<bool(const IBuildManifest&)> Predicate;
 		const TCHAR* FeatureLevelChunkSubdir;
 		FEvent* ManifestLoadedEvent;
 		TMap<uint64, TSet<FGuid>> ChunkInventory;
@@ -95,9 +97,10 @@ namespace BuildPatchServices
 		volatile FStatsCollector::FAtomicValue* StatUniqueWindowSizes;
 	};
 
-	FCloudEnumeration::FCloudEnumeration(const FString& InCloudDirectory, const FDateTime& InManifestAgeThreshold, const EFeatureLevel& InOutputFeatureLevel, FStatsCollector* InStatsCollector)
+	FCloudEnumeration::FCloudEnumeration(const FString& InCloudDirectory, const FDateTime& InManifestAgeThreshold, const TFunction<bool(const IBuildManifest&)>& InPredicate, const EFeatureLevel& InOutputFeatureLevel, FStatsCollector* InStatsCollector)
 		: CloudDirectory(InCloudDirectory)
 		, ManifestAgeThreshold(InManifestAgeThreshold)
+		, Predicate(InPredicate)
 		, FeatureLevelChunkSubdir(ManifestVersionHelpers::GetChunkSubdir(InOutputFeatureLevel))
 		, ManifestLoadedEvent(FPlatformProcess::GetSynchEventFromPool())
 		, StatsCollector(InStatsCollector)
@@ -294,6 +297,14 @@ namespace BuildPatchServices
 				const uint64 DeserializeEnd = FStatsCollector::GetCycles();
 				UE_LOG(LogCloudEnumeration, Log, TEXT("Deserialized manifest %s (%s in %.2f sec)"), *ManifestFilename, *FText::AsMemory(FileData.Num(), EMemoryUnitStandard::SI).ToString(), FStatsCollector::CyclesToSeconds(DeserializeEnd - DeserializeStart));
 				UE_LOG(LogCloudEnumeration, Log, TEXT("Manifest %s details: BuildVersion: %s, FeatureLevel: %s"), *ManifestFilename, *BuildManifest->GetVersionString(), LexToString(BuildManifest->GetFeatureLevel()));
+				// if predicate specified and returns false
+				if (Predicate && !Predicate(*BuildManifest))
+				{
+					FStatsCollector::Accumulate(StatManifestsRejected, 1);
+					UE_LOG(LogCloudEnumeration, Warning, TEXT("Manifest '%s' rejected by custom predicate."), *ManifestFilename);
+					return nullptr;
+				}
+
 				FStatsCollector::Accumulate(StatManifestsLoaded, 1);
 				return BuildManifest;
 			}
@@ -439,8 +450,8 @@ namespace BuildPatchServices
 		return DiscoveredDetails;
 	}
 
-	ICloudEnumeration* FCloudEnumerationFactory::Create(const FString& CloudDirectory, const FDateTime& ManifestAgeThreshold, const EFeatureLevel& OutputFeatureLevel, FStatsCollector* StatsCollector)
+	ICloudEnumeration* FCloudEnumerationFactory::Create(const FString& CloudDirectory, const FDateTime& ManifestAgeThreshold, const TFunction<bool(const IBuildManifest&)>& Predicate, const EFeatureLevel& OutputFeatureLevel, FStatsCollector* StatsCollector)
 	{
-		return new FCloudEnumeration(CloudDirectory, ManifestAgeThreshold, OutputFeatureLevel, StatsCollector);
+		return new FCloudEnumeration(CloudDirectory, ManifestAgeThreshold, Predicate ? Predicate : [](const IBuildManifest&) {return true; }, OutputFeatureLevel, StatsCollector);
 	}
 }
diff --git a/Engine/Source/Runtime/Online/BuildPatchServices/Private/Generation/CloudEnumeration.h b/Engine/Source/Runtime/Online/BuildPatchServices/Private/Generation/CloudEnumeration.h
--- a/Engine/Source/Runtime/Online/BuildPatchServices/Private/Generation/CloudEnumeration.h
+++ b/Engine/Source/Runtime/Online/BuildPatchServices/Private/Generation/CloudEnumeration.h
@@ -8,6 +8,8 @@
 #include "Common/StatsCollector.h"
 #include "BuildPatchFeatureLevel.h"
 
+class IBuildManifest;
+
 namespace BuildPatchServices
 {
 	class ICloudEnumeration
@@ -36,14 +38,15 @@ namespace BuildPatchServices
 	public:
 		/**
 		 * Creates an instance of ICloudEnumeration for managing cloud uploaded chunks.
-		 *
+		 * The predicate allows to filter manifests based on custom logic.
 		 * @param CloudDirectory The cloud root directory.
 		 * @param ManifestAgeThreshold The threshold age for manifest files; files older than this will be ignored.
+		 * @param Predicate A callable that takes a const reference to an IBuildManifest and returns true to include the manifest, or false to ignore it.
 		 * @param OutputFeatureLevel The desired feature level for enumerated output.
 		 * @param StatsCollector A pointer to the statistics collector for recording related data.
 		 *
 		 * @return A pointer to the created ICloudEnumeration instance.
 		 */
-		static ICloudEnumeration* Create(const FString& CloudDirectory, const FDateTime& ManifestAgeThreshold, const EFeatureLevel& OutputFeatureLevel, FStatsCollector* StatsCollector);
+		static ICloudEnumeration* Create(const FString& CloudDirectory, const FDateTime& ManifestAgeThreshold, const TFunction<bool(const IBuildManifest&)>& Predicate, const EFeatureLevel& OutputFeatureLevel, FStatsCollector* StatsCollector);
 	};
 }
diff --git a/Engine/Source/Runtime/Online/BuildPatchServices/Private/Generation/DeltaEnumeration.cpp b/Engine/Source/Runtime/Online/BuildPatchServices/Private/Generation/DeltaEnumeration.cpp
--- a/Engine/Source/Runtime/Online/BuildPatchServices/Private/Generation/DeltaEnumeration.cpp
+++ b/Engine/Source/Runtime/Online/BuildPatchServices/Private/Generation/DeltaEnumeration.cpp
@@ -38,7 +38,7 @@ namespace BuildPatchServices
 		: public IDeltaChunkEnumeration
 	{
 	public:
-		FDeltaChunkEnumeration(IManifestBuildStreamer* Streamer, FStatsCollector* StatsCollector, const FBuildPatchAppManifest& Manifest, const uint32 WindowSize, const TMap<FString, FString>& InAliasedTags);
+		FDeltaChunkEnumeration(IManifestBuildStreamer* Streamer, FStatsCollector* StatsCollector, const FBuildPatchAppManifest& Manifest, const uint32 WindowSize, const TMap<FString, FString>& InAliasedTags, const TSet<FString>& InIgnoredTags);
 		~FDeltaChunkEnumeration();
 
 		// FDeltaChunkEnumeration interface begin.
@@ -80,9 +80,10 @@ namespace BuildPatchServices
 		TMap<FDeltaChunkId, FChunkBuildReference> ChunkBuildReferences;
 		TMap<FShaId, TSet<FDeltaChunkId>> IdenticalChunks;
 		TMap<FString, FString> AliasedTags;
+		TSet<FString> IgnoredTags;
 	};
 
-	FDeltaChunkEnumeration::FDeltaChunkEnumeration(IManifestBuildStreamer* InStreamer, FStatsCollector* InStatsCollector, const FBuildPatchAppManifest& InManifest, const uint32 InWindowSize, const TMap<FString, FString>& InAliasedTags)
+	FDeltaChunkEnumeration::FDeltaChunkEnumeration(IManifestBuildStreamer* InStreamer, FStatsCollector* InStatsCollector, const FBuildPatchAppManifest& InManifest, const uint32 InWindowSize, const TMap<FString, FString>& InAliasedTags, const TSet<FString>& InIgnoredTags)
 		: Streamer(InStreamer)
 		, StatsCollector(InStatsCollector)
 		, ChunkSearcher(InManifest)
@@ -92,6 +93,7 @@ namespace BuildPatchServices
 		, bHasRan(false)
 		, bIsSuccess(true)
 		, AliasedTags(InAliasedTags)
+		, IgnoredTags(InIgnoredTags)
 	{
 		// Make sure we enumerate special chunk SHAs.
 		FDeltaChunkId PaddingChunkId = PaddingChunk::MakePaddingGuid(0);
@@ -316,10 +318,14 @@ namespace BuildPatchServices
 				FChunkBuildReferenceSourceFile& SourceFile = ChunkBuildReference.FileReferences.AddDefaulted_GetRef();
 				SourceFile.Filename = MakeFilenameId(FileNode.Manifest->Filename);
 				SourceFile.Location = (ChunkNode.BuildRange.GetFirst() - FileNode.BuildRange.GetFirst()) + ChunkLeftChop;
-				// If we have any tag aliases, swap out the tags so that they all match the same.
+				// If we have any tag aliases, swap out the tags so that they all match the same, and remove any tags that shouldn't be matched against in the matcher.
 				SourceFile.Tagset.Reserve(FileNode.Manifest->InstallTags.Num());
 				for (const FString& InstallTag : FileNode.Manifest->InstallTags)
 				{
+					if (IgnoredTags.Contains(InstallTag))
+					{
+						continue;
+					}
 					const FString* AliasedTag = AliasedTags.Find(InstallTag);
 					SourceFile.Tagset.Add(AliasedTag ? *AliasedTag : InstallTag);
 				}
@@ -336,8 +342,8 @@ namespace BuildPatchServices
 		return SHAHash;
 	}
 
-	IDeltaChunkEnumeration* FDeltaChunkEnumerationFactory::Create(IManifestBuildStreamer* Streamer, FStatsCollector* StatsCollector, const FBuildPatchAppManifest& Manifest, const uint32 WindowSize, const TMap<FString, FString>& AliasedTags)
+	IDeltaChunkEnumeration* FDeltaChunkEnumerationFactory::Create(IManifestBuildStreamer* Streamer, FStatsCollector* StatsCollector, const FBuildPatchAppManifest& Manifest, const uint32 WindowSize, const TMap<FString, FString>& AliasedTags, const TSet<FString>& IgnoredTags)
 	{
-		return new FDeltaChunkEnumeration(Streamer, StatsCollector, Manifest, WindowSize, AliasedTags);
+		return new FDeltaChunkEnumeration(Streamer, StatsCollector, Manifest, WindowSize, AliasedTags, IgnoredTags);
 	}
 }
diff --git a/Engine/Source/Runtime/Online/BuildPatchServices/Private/Generation/DeltaEnumeration.h b/Engine/Source/Runtime/Online/BuildPatchServices/Private/Generation/DeltaEnumeration.h
--- a/Engine/Source/Runtime/Online/BuildPatchServices/Private/Generation/DeltaEnumeration.h
+++ b/Engine/Source/Runtime/Online/BuildPatchServices/Private/Generation/DeltaEnumeration.h
@@ -57,6 +57,6 @@ namespace BuildPatchServices
 	class FDeltaChunkEnumerationFactory
 	{
 	public:
-		static IDeltaChunkEnumeration* Create(IManifestBuildStreamer* Streamer, FStatsCollector* StatsCollector, const FBuildPatchAppManifest& Manifest, const uint32 WindowSize, const TMap<FString, FString>& InAliasedTags);
+		static IDeltaChunkEnumeration* Create(IManifestBuildStreamer* Streamer, FStatsCollector* StatsCollector, const FBuildPatchAppManifest& Manifest, const uint32 WindowSize, const TMap<FString, FString>& InAliasedTags, const TSet<FString>& InIgnoredTags);
 	};
 }
diff --git a/Engine/Source/Runtime/Online/BuildPatchServices/Private/Generation/PackageChunkData.cpp b/Engine/Source/Runtime/Online/BuildPatchServices/Private/Generation/PackageChunkData.cpp
--- a/Engine/Source/Runtime/Online/BuildPatchServices/Private/Generation/PackageChunkData.cpp
+++ b/Engine/Source/Runtime/Online/BuildPatchServices/Private/Generation/PackageChunkData.cpp
@@ -201,10 +201,10 @@ namespace BuildPatchServices
 		TUniquePtr<ICrypto> Crypto;
 		TUniquePtr<IMessagePump> MessagePump;
 		TUniquePtr<IInstallerError> InstallerError;
-		TUniquePtr<ISpeedRecorder> DownloadSpeedRecorder;
+		TSharedPtr<ISpeedRecorder> DownloadSpeedRecorder;
 		TUniquePtr<ISpeedRecorder> DiskReadSpeedRecorder;
 		TUniquePtr<ISpeedRecorder> DiskWriteSpeedRecorder;
-		TUniquePtr<IChunkDataSizeProvider> ChunkDataSizeProvider;
+		TSharedPtr<IChunkDataSizeProvider> ChunkDataSizeProvider;
 		TSharedPtr<IInstallerAnalytics> InstallerAnalytics;
 		TSharedPtr<IDownloadServiceStatistics> DownloadServiceStatistics;
 		TUniquePtr<IDownloadService> DownloadService;
@@ -250,7 +250,7 @@ namespace BuildPatchServices
 		, DiskWriteSpeedRecorder(FSpeedRecorderFactory::Create())
 		, ChunkDataSizeProvider(FChunkDataSizeProviderFactory::Create({}))
 		, InstallerAnalytics(FInstallerAnalyticsFactory::Create(nullptr))
-		, DownloadServiceStatistics(FDownloadServiceStatisticsFactory::Create(DownloadSpeedRecorder.Get(), ChunkDataSizeProvider.Get(), InstallerAnalytics.Get()))
+		, DownloadServiceStatistics(FDownloadServiceStatisticsFactory::Create(DownloadSpeedRecorder, ChunkDataSizeProvider, InstallerAnalytics))
 		, DownloadService(FDownloadServiceFactory::Create(HttpManager, FileSystem, DownloadServiceStatistics, InstallerAnalytics))
 		, FileOperationTracker(FFileOperationTrackerFactory::Create(CoreTicker))
 	{
diff --git a/Engine/Source/Runtime/Online/BuildPatchServices/Private/IBuildManifestSet.h b/Engine/Source/Runtime/Online/BuildPatchServices/Private/IBuildManifestSet.h
--- a/Engine/Source/Runtime/Online/BuildPatchServices/Private/IBuildManifestSet.h
+++ b/Engine/Source/Runtime/Online/BuildPatchServices/Private/IBuildManifestSet.h
@@ -49,9 +49,9 @@ namespace BuildPatchServices
 			// Always require the empty tag.
 			InstallTags.Add(TEXT(""));
 			// Cache tagged files.
-			if (!IsUninstall())
+			if (InstallManifest.IsValid())
 			{
-				GetInstallManifest().GetTaggedFileList(InstallTags, TaggedFiles);
+				InstallManifest->GetTaggedFileList(InstallTags, TaggedFiles);
 			}
 		}
 
@@ -61,7 +61,6 @@ namespace BuildPatchServices
 		bool IsVerifyOnly() const { return InstallerAction.IsVerifyOnly(); }
 		bool IsUninstall() const { return InstallerAction.IsUninstall(); }
 		const TSet<FString>& GetInstallTags() const { return InstallTags; }
-		const TSet<FString>& GetTaggedFiles() const { return TaggedFiles; }
 		const FString& GetInstallSubdirectory() const { return InstallerAction.GetInstallSubdirectory(); }
 		const FString& GetCloudSubdirectory() const { return InstallerAction.GetCloudSubdirectory(); }
 		FString GetUniqueBuildId() const { return InstallManifest.IsValid() ? InstallManifest->GetUniqueBuildId() : FString(); }
@@ -79,6 +78,12 @@ namespace BuildPatchServices
 		const FBuildPatchAppManifestRef GetSharedInstallOrCurrentManifest() const { return (InstallManifest.IsValid() ? InstallManifest : CurrentManifest).ToSharedRef(); }
 
 		void SetDeltaManifest(const FBuildPatchAppManifestRef& DeltaManifest) { if (IsUpdate()) { InstallManifest = DeltaManifest; } }
+
+	private:
+		// Hide tagged file accessors, they are only expected to be used by the manifest set, to prevent other systems from confusing
+		// flat manifest filenames, and filenames including the action subdirectory.
+		friend class FBuildPatchManifestSet;
+		const TSet<FString>& GetTaggedFiles() const { return TaggedFiles; }
 		void SetTaggedFiles(TSet<FString> OverrideTaggedFiles) { TaggedFiles = MoveTemp(OverrideTaggedFiles); }
 
 	private:
@@ -105,10 +110,10 @@ namespace BuildPatchServices
 		virtual uint64 GetDownloadSize(const TSet<FGuid>& DataGuids) const = 0;
 		virtual FChunkInfo const* GetChunkInfo(const FGuid& DataGuid) const = 0;
 		virtual bool GetChunkShaHash(const FGuid& ChunkGuid, FSHAHash& OutHash) const = 0;
-		virtual FString GetDataFilename(const FGuid& DataGuid) const = 0;
 		virtual FGuid GetDataSecretId(const FGuid& DataGuid) const = 0;
 		virtual FString GetChunkCloudSubdirectory(const FGuid& DataGuid) const = 0;
 		virtual FString GetChunkFilename(const FGuid& DataGuid) const = 0;
+		virtual FString GetChunkFilenameWithCloudSubdirectory(const FGuid& DataGuid) const = 0;
 		virtual FString GetChunkUniqueBuildId(const FGuid& DataGuid) const = 0;
 		virtual bool GetChunkFromDeltaStatus(const FGuid& DataGuid) const = 0;
 
@@ -123,7 +128,6 @@ namespace BuildPatchServices
 		virtual const FFileManifest* GetNewFileManifest(const FString& BuildFile) const = 0;
 		virtual uint64 GetTotalNewFileSize(const TArray<FString>& Filenames) const = 0;
 		virtual uint64 GetTotalNewFileSize(const TSet<FString>& Filenames) const = 0;
-		virtual void GetPreReqInfo(TArray<FPreReqInfo>& PreReqInfo) const = 0;
 		virtual void GetFilesTaggedForRepair(TSet<FString>& Filenames) const = 0;
 		virtual bool IsFileRepairAction(const FString& Filename) const = 0;
 
@@ -133,7 +137,8 @@ namespace BuildPatchServices
 		virtual bool HasFileAttributes() const = 0;
 		virtual bool IsExclusivelyVerifying() const = 0;
 
-		// Custom Uninstall Action support
+		// Custom install support
+		virtual void GetPreReqInfo(TArray<FPreReqInfo>& PreReqInfo) const = 0;
 		virtual void GetUninstallActionInfo(TArray<FUninstallActionInfo>& UninstallActionsInfo) const = 0;
 	};
 
@@ -142,4 +147,4 @@ namespace BuildPatchServices
 	public:
 		static IBuildManifestSet* Create(TArray<FBuildPatchInstallerAction> InstallerActions);
 	};
-}
+}
\ No newline at end of file
diff --git a/Engine/Source/Runtime/Online/BuildPatchServices/Private/Installer/ChunkDbChunkSource.cpp b/Engine/Source/Runtime/Online/BuildPatchServices/Private/Installer/ChunkDbChunkSource.cpp
--- a/Engine/Source/Runtime/Online/BuildPatchServices/Private/Installer/ChunkDbChunkSource.cpp
+++ b/Engine/Source/Runtime/Online/BuildPatchServices/Private/Installer/ChunkDbChunkSource.cpp
@@ -22,6 +22,7 @@
 #include "Serialization/ArrayWriter.h"
 #include "Misc/Compression.h"
 #include "Compression/OodleDataCompressionUtil.h"
+#include "Misc/Paths.h"
 
 #if PLATFORM_WINDOWS
 #include "Windows/WindowsHWrapper.h"
@@ -105,7 +106,7 @@ namespace BuildPatchServices
 			return Lock;
 		}
 
-		void Retire(IFileSystem* FileSystemIfDeleting, bool bDelete)
+		void Retire(IFileSystem* FileSystem, bool bDelete)
 		{
 			bIsRetired = true;
 			
@@ -113,11 +114,15 @@ namespace BuildPatchServices
 				Archive.Reset();
 				FileHandler = nullptr;
 			}
-			if (bDelete && FileSystemIfDeleting)
+			if (bDelete && FileSystem)
 			{
-				if (!FileSystemIfDeleting->DeleteFile(*ChunkDbFileName))
+				if (FileSystem->DeleteFile(*ChunkDbFileName))
 				{
-					GLog->Logf(TEXT("Failed to delete chunkdb upon retirement: %s"), *ChunkDbFileName);
+					UE_LOG(LogChunkDbChunkSource, Display, TEXT("Chunkdb deleted upon retirement: %s"), *FPaths::GetCleanFilename(ChunkDbFileName));
+				}
+				else
+				{
+					UE_LOG(LogChunkDbChunkSource, Error, TEXT("Failed to delete chunkdb upon retirement: %s"), *ChunkDbFileName);
 				}
 			}
 		}
@@ -176,31 +181,33 @@ namespace BuildPatchServices
 			//
 			for (FChunkDbDataAccess& ChunkDbDataAccess : ChunkDbDataAccesses)
 			{
-				if (ChunkDbDataAccess.ChunkRetiresAt.Num() > 0)
+				if (ChunkDbDataAccess.ChunkRetiresAt.Num() == 0)
 				{
-					const int32 CurrentLastChunk = ChunkDbDataAccess.ChunkRetiresAt.Num() - 1;
-					int32 LastChunkRequired = CurrentLastChunk;
-					for (; LastChunkRequired >= 0; LastChunkRequired--)
-					{
-						if (ChunkDbDataAccess.ChunkRetiresAt[LastChunkRequired] <= RemainingChunkCount)
-						{
-							break;
-						}
-					}
-					if (LastChunkRequired != CurrentLastChunk)
-					{
-						ChunkDbDataAccess.ChunkRetiresAt.SetNum(LastChunkRequired + 1, EAllowShrinking::No);
-						int64 RequiredChunkDbFileSize = 0;
-						if (LastChunkRequired >= 0)
-						{
-							const FChunkLocation& LastChunkLocation = ChunkDbDataAccess.Header.Contents[LastChunkRequired];
-							RequiredChunkDbFileSize = LastChunkLocation.ByteStart + LastChunkLocation.ByteSize;
-						}
-						ChunkDbDataAccess.TruncateIfNecessary(FileSystem, RequiredChunkDbFileSize);
-					}
+					continue;
 				}
+				
+				const int32 LastChunkPartRequired = ChunkDbDataAccess.ChunkRetiresAt.FindLastByPredicate([RemainingChunkCount](int32 Value) { return Value <= RemainingChunkCount; });
+				if (LastChunkPartRequired == (ChunkDbDataAccess.ChunkRetiresAt.Num() - 1))
+				{
+					continue;
+				}
+
+				ChunkDbDataAccess.ChunkRetiresAt.SetNum(LastChunkPartRequired + 1, EAllowShrinking::No);
+				int64 RequiredChunkDbFileSize = 0;
+				if (LastChunkPartRequired != INDEX_NONE)
+				{
+					const FChunkLocation& LastChunkLocation = ChunkDbDataAccess.Header.Contents[LastChunkPartRequired];
+					RequiredChunkDbFileSize = LastChunkLocation.ByteStart + LastChunkLocation.ByteSize;
+					UE_LOG(LogChunkDbChunkSource, Log, TEXT("Truncating file %s up to chunk part %d (tracker index: %d), file size: %lld. Remaining chunk parts in tracker: %d"), *ChunkDbDataAccess.ChunkDbFileName, LastChunkPartRequired, ChunkDbDataAccess.ChunkRetiresAt[LastChunkPartRequired], RequiredChunkDbFileSize, RemainingChunkCount);
+				}
+				else
+				{
+					UE_LOG(LogChunkDbChunkSource, Log, TEXT("Truncating file %s up to chunk part %d (tracker index: <unknown>), file size: %lld. Remaining chunk parts in tracker: %d"), *ChunkDbDataAccess.ChunkDbFileName, LastChunkPartRequired, RequiredChunkDbFileSize, RemainingChunkCount);
+				}
+				ChunkDbDataAccess.TruncateIfNecessary(FileSystem, RequiredChunkDbFileSize);
 			}
 		}
+
 		virtual int32 GetChunkUnavailableAt(const FGuid& DataId) const override
 		{
 			// While technically the chunks retire as a result of delete-during-install, we only do this
@@ -319,8 +326,7 @@ namespace BuildPatchServices
 
 		TMap<FString, int32> FileLastSeenAt;
 
-		int32 GuidIndex = 0;
-		for (; GuidIndex < ChunkAccessOrderedList.Num(); GuidIndex++)
+		for (int32 GuidIndex = 0; GuidIndex < ChunkAccessOrderedList.Num(); GuidIndex++)
 		{
 			const FGuid& Guid = ChunkAccessOrderedList[GuidIndex];
 
@@ -335,9 +341,9 @@ namespace BuildPatchServices
 			if (bTruncateChunkDbs)
 			{
 				const TArray<FChunkLocation>& ChunkDBContents = SourceForGuid->DbFile->Header.Contents;
-				const int32 IndexOfChunk = SourceForGuid->Location - &ChunkDBContents[0];
-				check(ChunkDBContents.IsValidIndex(IndexOfChunk)); // programmer error if this fires
-				SourceForGuid->DbFile->ChunkRetiresAt[IndexOfChunk] = ChunkAccessOrderedList.Num() - (GuidIndex + 1);
+				const int32 IndexOfChunkPart = SourceForGuid->Location - ChunkDBContents.GetData();
+				check(ChunkDBContents.IsValidIndex(IndexOfChunkPart)); // programmer error if this fires
+				SourceForGuid->DbFile->ChunkRetiresAt[IndexOfChunkPart] = ChunkAccessOrderedList.Num() - (GuidIndex + 1);
 			}
 		}
 
diff --git a/Engine/Source/Runtime/Online/BuildPatchServices/Private/Installer/CloudChunkSource.cpp b/Engine/Source/Runtime/Online/BuildPatchServices/Private/Installer/CloudChunkSource.cpp
--- a/Engine/Source/Runtime/Online/BuildPatchServices/Private/Installer/CloudChunkSource.cpp
+++ b/Engine/Source/Runtime/Online/BuildPatchServices/Private/Installer/CloudChunkSource.cpp
@@ -1549,7 +1549,8 @@ namespace BuildPatchServices
 
 					if (!Result)
 					{
-						// The header or chunk data was bad.				
+						// The header or chunk data was bad.
+						UE_LOG(LogCloudChunkSource, Error, TEXT("Failed to ValidateAndRead chunk %s, download size %d"), *WriteToString<64>(Read->DataId), DownloadedData.Num());
 						bDownloadSuccess = false;
 					}
 					else
diff --git a/Engine/Source/Runtime/Online/BuildPatchServices/Private/Installer/DownloadService.cpp b/Engine/Source/Runtime/Online/BuildPatchServices/Private/Installer/DownloadService.cpp
--- a/Engine/Source/Runtime/Online/BuildPatchServices/Private/Installer/DownloadService.cpp
+++ b/Engine/Source/Runtime/Online/BuildPatchServices/Private/Installer/DownloadService.cpp
@@ -101,7 +101,7 @@ namespace BuildPatchServices
 	{
 	public:
 		FDownloadServiceImpl(TSharedPtr<IHttpManager> HttpManager, TSharedPtr<IFileSystem> FileSystem, TSharedPtr<IDownloadServiceStat> DownloadServiceStat, TSharedPtr<IInstallerAnalytics> InstallerAnalytics);
-		~FDownloadServiceImpl();
+		~FDownloadServiceImpl() = default;
 
 		void CancelAllRequests();
 
@@ -116,6 +116,10 @@ namespace BuildPatchServices
 
 		IDownloadServiceStat::FDownloadRecord MakeDownloadRecord(int32 RequestId, FString Uri);
 
+		void OnRequestProgress(FHttpRequestPtr Request, uint64 BytesSent, uint64 BytesReceived, const int32 RequestId);
+
+		void OnProcessRequestComplete(FHttpRequestPtr Request, FHttpResponsePtr Response, bool bSucceeded, IDownloadServiceStat::FDownloadRecord DownloadRecord);
+
 	private:
 		TSharedPtr<IHttpManager> HttpManager;
 		TSharedPtr<IFileSystem> FileSystem;
@@ -169,10 +173,6 @@ namespace BuildPatchServices
 		, InstallerAnalytics(InInstallerAnalytics)
 	{}
 
-	FDownloadServiceImpl::~FDownloadServiceImpl()
-	{
-	}
-
 	void FDownloadServiceImpl::CancelAllRequests()
 	{
 		// Cancel all HTTP requests and wait for all file downloads threads to exit.
@@ -224,57 +224,10 @@ namespace BuildPatchServices
 			{
 				HttpRequest->SetHeader(Header.Key, Header.Value);
 			}
-			
-			HttpRequest->OnRequestProgress64().BindSPLambda(this, [this, RequestId](FHttpRequestPtr Request, uint64 BytesSent, uint64 BytesReceived)
-			{
-				#if USING_ADDRESS_SANITISER
-					// Force use the Request parameter to work around some MSVC-specific compiler issue in ASan.
-					UE_LOG(LogDownloadService, Verbose, TEXT("[FDownloadServiceImpl::RequestFile] %s %s received %d bytes"), *Request->GetVerb(), *Request->GetURL(), BytesReceived);
-				#endif
-
-				{
-					FScopeLock ScopeLock(&RequestDelegatesCS);
-					if (FDownloadDelegates* MaybeDelegates = RequestDelegates.Find(RequestId))
-					{
-						MaybeDelegates->OnProgressDelegate.ExecuteIfBound(RequestId, BytesReceived);
-					}
-				}
-
-				DownloadServiceStat->OnDownloadProgress(RequestId, BytesReceived);
-			});
-
-			HttpRequest->OnProcessRequestComplete().BindSPLambda(this,
-			[this, DownloadRecord = MakeDownloadRecord(RequestId, FileUri)](FHttpRequestPtr Request, FHttpResponsePtr Response, bool bSucceeded) mutable
-			{
-				InstallerAnalytics->TrackRequest(Request);
-
-				DownloadRecord.bSuccess = bSucceeded;
-				DownloadRecord.ResponseCode = Response.IsValid() ? Response->GetResponseCode() : (int32)Request->GetFailureReason();
-				DownloadRecord.SpeedRecord.CyclesEnd = FStatsCollector::GetCycles();
-				DownloadRecord.SpeedRecord.Size = Response.IsValid() ? Response->GetContent().Num() : 0;
-				DownloadRecord.FailureReason = Request->GetFailureReason();
-				
-				FDownloadDelegates Delegates;
-				bool bFoundDelegates = false;
-				{
-					FScopeLock ScopeLock(&RequestDelegatesCS);
-					bFoundDelegates = RequestDelegates.RemoveAndCopyValue(DownloadRecord.RequestId, Delegates);
-				}
-
-				if (bFoundDelegates)
-				{
-					TSharedRef<FHttpDownload> CompletedDownload = MakeShared<FHttpDownload>(Request, bSucceeded);
-					Delegates.OnCompleteDelegate.ExecuteIfBound(DownloadRecord.RequestId, CompletedDownload);
-				}
-
-				{
-					FScopeLock ScopeLock(&ActiveRequestsCS);
-					ActiveHttpRequests.Remove(DownloadRecord.RequestId);
-				}
-
-				DownloadServiceStat->OnDownloadComplete(DownloadRecord);
-			});
 
+			// We have to use BingSP here because BindSPLambda is not making sure that the context object is alive till the end of callback execution.
+			HttpRequest->OnRequestProgress64().BindSP(this, &FDownloadServiceImpl::OnRequestProgress, RequestId);
+			HttpRequest->OnProcessRequestComplete().BindSP(this, &FDownloadServiceImpl::OnProcessRequestComplete, MakeDownloadRecord(RequestId, FileUri));
 			HttpRequest->ProcessRequest();
 		}
 		else
@@ -423,6 +376,54 @@ namespace BuildPatchServices
 		return DownloadRecord;
 	}
 
+	void FDownloadServiceImpl::OnRequestProgress(FHttpRequestPtr Request, uint64 BytesSent, uint64 BytesReceived, const int32 RequestId)
+	{
+#if USING_ADDRESS_SANITISER
+		// Force use the Request parameter to work around some MSVC-specific compiler issue in ASan.
+		UE_LOG(LogDownloadService, Verbose, TEXT("[FDownloadServiceImpl::RequestFile] %s %s received %d bytes"), *Request->GetVerb(), *Request->GetURL(), BytesReceived);
+#endif
+		{
+			FScopeLock ScopeLock(&RequestDelegatesCS);
+			if (FDownloadDelegates* MaybeDelegates = RequestDelegates.Find(RequestId))
+			{
+				MaybeDelegates->OnProgressDelegate.ExecuteIfBound(RequestId, BytesReceived);
+			}
+		}
+
+		DownloadServiceStat->OnDownloadProgress(RequestId, BytesReceived);
+	}
+
+	void FDownloadServiceImpl::OnProcessRequestComplete(FHttpRequestPtr Request, FHttpResponsePtr Response, bool bSucceeded, IDownloadServiceStat::FDownloadRecord DownloadRecord)
+	{
+		InstallerAnalytics->TrackRequest(Request);
+
+		DownloadRecord.bSuccess = bSucceeded;
+		DownloadRecord.ResponseCode = Response.IsValid() ? Response->GetResponseCode() : (int32)Request->GetFailureReason();
+		DownloadRecord.SpeedRecord.CyclesEnd = FStatsCollector::GetCycles();
+		DownloadRecord.SpeedRecord.Size = Response.IsValid() ? Response->GetContent().Num() : 0;
+		DownloadRecord.FailureReason = Request->GetFailureReason();
+
+		FDownloadDelegates Delegates;
+		bool bFoundDelegates = false;
+		{
+			FScopeLock ScopeLock(&RequestDelegatesCS);
+			bFoundDelegates = RequestDelegates.RemoveAndCopyValue(DownloadRecord.RequestId, Delegates);
+		}
+
+		if (bFoundDelegates)
+		{
+			TSharedRef<FHttpDownload> CompletedDownload = MakeShared<FHttpDownload>(Request, bSucceeded);
+			Delegates.OnCompleteDelegate.ExecuteIfBound(DownloadRecord.RequestId, CompletedDownload);
+		}
+
+		{
+			FScopeLock ScopeLock(&ActiveRequestsCS);
+			ActiveHttpRequests.Remove(DownloadRecord.RequestId);
+		}
+
+		DownloadServiceStat->OnDownloadComplete(DownloadRecord);
+	}
+
 	IDownloadService* FDownloadServiceFactory::Create(TSharedPtr<IHttpManager> HttpManager, TSharedPtr<IFileSystem> FileSystem, TSharedPtr<IDownloadServiceStat> DownloadServiceStat, TSharedPtr<IInstallerAnalytics> InstallerAnalytics)
 	{
 		check(HttpManager != nullptr);
diff --git a/Engine/Source/Runtime/Online/BuildPatchServices/Private/Installer/InstallerAnalytics.cpp b/Engine/Source/Runtime/Online/BuildPatchServices/Private/Installer/InstallerAnalytics.cpp
--- a/Engine/Source/Runtime/Online/BuildPatchServices/Private/Installer/InstallerAnalytics.cpp
+++ b/Engine/Source/Runtime/Online/BuildPatchServices/Private/Installer/InstallerAnalytics.cpp
@@ -40,6 +40,7 @@ namespace BuildPatchServices
 	public:
 		// IInstallerAnalytics interface begin.
 		virtual void RecordChunkDownloadError(const FString& ChunkUrl, int32 ResponseCode, const FString& ErrorString) override {}
+		virtual void RecordMissingDelta(const FString& DeltaUrl, const FString& BaseVersion, const FString& PatchVersion) override {}
 		virtual void RecordChunkDownloadAborted(const FString& ChunkUrl, double ChunkTime, double ChunkMean, double ChunkStd, double BreakingPoint) override {}
 		virtual void RecordChunkCacheError(const FGuid& ChunkGuid, const FString& Filename, int32 LastError, const FString& SystemName, const FString& ErrorString) override {}
 		virtual void RecordConstructionError(const FString& Filename, int32 LastError, const FString& ErrorString) override {}
@@ -58,6 +59,7 @@ namespace BuildPatchServices
 
 		// IInstallerAnalytics interface begin.
 		virtual void RecordChunkDownloadError(const FString& ChunkUrl, int32 ResponseCode, const FString& ErrorString) override;
+		virtual void RecordMissingDelta(const FString& DeltaUrl, const FString& BaseVersion, const FString& PatchVersion) override;
 		virtual void RecordChunkDownloadAborted(const FString& ChunkUrl, double ChunkTime, double ChunkMean, double ChunkStd, double BreakingPoint) override;
 		virtual void RecordChunkCacheError(const FGuid& ChunkGuid, const FString& Filename, int32 LastError, const FString& SystemName, const FString& ErrorString) override;
 		virtual void RecordConstructionError(const FString& Filename, int32 LastError, const FString& ErrorString) override;
@@ -110,6 +112,15 @@ namespace BuildPatchServices
 		}
 	}
 
+	void FInstallerAnalytics::RecordMissingDelta(const FString& DeltaUrl, const FString& BaseVersion, const FString& PatchVersion)
+	{
+		TArray<FAnalyticsEventAttribute> Attributes;
+		Attributes.Add(FAnalyticsEventAttribute(TEXT("DeltaUrl"), DeltaUrl));
+		Attributes.Add(FAnalyticsEventAttribute(TEXT("BaseVersion"), BaseVersion));
+		Attributes.Add(FAnalyticsEventAttribute(TEXT("PatchVersion"), PatchVersion));
+		QueueAnalyticsEvent(TEXT("Patcher.Error.Delta"), MoveTemp(Attributes));
+	}
+
 	void FInstallerAnalytics::RecordChunkDownloadAborted(const FString& ChunkUrl, double ChunkTime, double ChunkMean, double ChunkStd, double BreakingPoint)
 	{
 		TArray<FAnalyticsEventAttribute> Attributes;
diff --git a/Engine/Source/Runtime/Online/BuildPatchServices/Private/Installer/InstallerAnalytics.h b/Engine/Source/Runtime/Online/BuildPatchServices/Private/Installer/InstallerAnalytics.h
--- a/Engine/Source/Runtime/Online/BuildPatchServices/Private/Installer/InstallerAnalytics.h
+++ b/Engine/Source/Runtime/Online/BuildPatchServices/Private/Installer/InstallerAnalytics.h
@@ -33,6 +33,20 @@ namespace BuildPatchServices
 		 */
 		virtual void RecordChunkDownloadError(const FString& ChunkUrl, int32 ResponseCode, const FString& ErrorString) = 0;
 
+		/**
+		 * @EventName Patcher.Error.Delta
+		 *
+		 * @Trigger Fires for delta files that return 403/404 http response codes
+		 *
+		 * @Type Client
+		 *
+		 * @EventParam DeltaUrl		(string) The url used to request the delta file
+		 * @EventParam BaseVersion	(string) The base (installed) version of the app
+		 * @EventParam PatchVersion	(string) The version of the app that is being upgraded to
+		 * @Owner Portal Team
+		 */
+		virtual void RecordMissingDelta(const FString& DeltaUrl, const FString& BaseVersion, const FString& PatchVersion) = 0;
+
 		/**
 		 * @EventName Patcher.Warning.ChunkAborted
 		 *
diff --git a/Engine/Source/Runtime/Online/BuildPatchServices/Private/Installer/InstallerError.cpp b/Engine/Source/Runtime/Online/BuildPatchServices/Private/Installer/InstallerError.cpp
--- a/Engine/Source/Runtime/Online/BuildPatchServices/Private/Installer/InstallerError.cpp
+++ b/Engine/Source/Runtime/Online/BuildPatchServices/Private/Installer/InstallerError.cpp
@@ -79,7 +79,7 @@ namespace BuildPatchServices
 		virtual void SetError(EBuildPatchInstallError InErrorType, const TCHAR* InErrorSubType, uint32 InErrorCode, FText InErrorText) override;
 		virtual int32 RegisterForErrors(FOnErrorDelegate Delegate) override;
 		virtual void UnregisterForErrors(int32 Handle) override;
-		virtual void Reset() override;
+		virtual void ResetError(bool bPreserveCancelled) override;
 		// IInstallerError interface end.
 
 	private:
@@ -186,9 +186,13 @@ namespace BuildPatchServices
 		RegisteredDelegates.Remove(Handle);
 	}
 
-	void FInstallerError::Reset()
+	void FInstallerError::ResetError(bool bPreserveCancelled)
 	{
 		FScopeLock ScopeLock(&ThreadLockCs);
+		if (bPreserveCancelled && IsCancelled())
+		{
+			return;
+		}
 		ErrorType = EBuildPatchInstallError::NoError;
 		ErrorCode = InstallErrorPrefixes::ErrorTypeStrings[static_cast<int32>(ErrorType)];
 		ErrorText = GetStandardErrorText(ErrorType);
diff --git a/Engine/Source/Runtime/Online/BuildPatchServices/Private/Installer/InstallerError.h b/Engine/Source/Runtime/Online/BuildPatchServices/Private/Installer/InstallerError.h
--- a/Engine/Source/Runtime/Online/BuildPatchServices/Private/Installer/InstallerError.h
+++ b/Engine/Source/Runtime/Online/BuildPatchServices/Private/Installer/InstallerError.h
@@ -206,6 +206,12 @@ namespace BuildPatchServices
 		 */
 		virtual void SetError(EBuildPatchInstallError ErrorType, const TCHAR* ErrorSubType, uint32 ErrorCode = 0, FText ErrorText = FText::GetEmpty()) = 0;
 
+		/**
+		 * Resets the current error state without unregistering delegates added by RegisterForErrors.
+		 * @param bPreserveCancelled	If true the canceled error state is preserved instead of being reset.
+		 */
+		virtual void ResetError(bool bPreserveCancelled) = 0;
+
 		/**
 		 * Register a delegate to be called upon an error occurring.
 		 * @param Delegate      The delegate to be executed on error.
@@ -218,11 +224,6 @@ namespace BuildPatchServices
 		 * @param Handle        The value that was returned RegisterForErrors for the delegate to be removed.
 		 */
 		virtual void UnregisterForErrors(int32 Handle) = 0;
-
-		/**
-		 * Reset internal state to be ready for a new process.
-		 */
-		virtual void Reset() = 0;
 	};
 
 	/**
diff --git a/Engine/Source/Runtime/Online/BuildPatchServices/Private/Installer/OptimisedDelta.cpp b/Engine/Source/Runtime/Online/BuildPatchServices/Private/Installer/OptimisedDelta.cpp
--- a/Engine/Source/Runtime/Online/BuildPatchServices/Private/Installer/OptimisedDelta.cpp
+++ b/Engine/Source/Runtime/Online/BuildPatchServices/Private/Installer/OptimisedDelta.cpp
@@ -84,6 +84,7 @@ namespace BuildPatchServices
 		virtual const IOptimisedDelta::FResultValueOrError& GetResult() const override;
 		virtual int32 GetMetaDownloadSize() const override;
 		virtual bool IsReady() const override;
+		int32 GetDownloadResponseCode() const override;
 		// IOptimisedDelta interface end.
 
 		void Initialize();
@@ -115,6 +116,7 @@ namespace BuildPatchServices
 		TFuture<TSharedPtr<FResultValueOrError, ESPMode::ThreadSafe>> ChunkDeltaFuture{ ChunkDeltaPromise.GetFuture() };
 		FThreadSafeCounter DownloadedBytes = 0;
 		const TCHAR* ErrorCode = nullptr;
+		int32 DownloadResponseCode = 0;
 	};
 
 	FOptimisedDelta::FOptimisedDelta(const FOptimisedDeltaConfiguration& InConfiguration, FOptimisedDeltaDependencies InDependencies)
@@ -194,6 +196,7 @@ namespace BuildPatchServices
 
 	void FOptimisedDelta::OnDownloadComplete(int32 RequestId, const FDownloadRef& Download)
 	{
+		DownloadResponseCode = Download->GetResponseCode();		// Store for analytics.
 		if (Download->ResponseSuccessful())
 		{
 			// Perform a merge with current manifest so that the delta can support missing out unnecessary information.
@@ -351,6 +354,11 @@ namespace BuildPatchServices
 		}
 	}
 
+	int32 FOptimisedDelta::GetDownloadResponseCode() const
+	{
+		return DownloadResponseCode;
+	}
+
 	FOptimisedDeltaConfiguration::FOptimisedDeltaConfiguration(FBuildPatchAppManifestRef InDestinationManifest)
 		: DestinationManifest(MoveTemp(InDestinationManifest))
 		, DeltaPolicy(EDeltaPolicy::TryFetchContinueWithout)
diff --git a/Engine/Source/Runtime/Online/BuildPatchServices/Private/Installer/OptimisedDelta.h b/Engine/Source/Runtime/Online/BuildPatchServices/Private/Installer/OptimisedDelta.h
--- a/Engine/Source/Runtime/Online/BuildPatchServices/Private/Installer/OptimisedDelta.h
+++ b/Engine/Source/Runtime/Online/BuildPatchServices/Private/Installer/OptimisedDelta.h
@@ -39,6 +39,11 @@ namespace BuildPatchServices
 		 */
 		virtual bool IsReady() const = 0;
 
+		/**
+		 * Gets the response code from the download of the optimised delta file.
+		 * @return the response code, or 0 if no download was attempted.
+		 */
+		virtual int32 GetDownloadResponseCode() const = 0;
 	};
 
 	/**
diff --git a/Engine/Source/Runtime/Online/BuildPatchServices/Private/Installer/Statistics/DownloadServiceStatistics.cpp b/Engine/Source/Runtime/Online/BuildPatchServices/Private/Installer/Statistics/DownloadServiceStatistics.cpp
--- a/Engine/Source/Runtime/Online/BuildPatchServices/Private/Installer/Statistics/DownloadServiceStatistics.cpp
+++ b/Engine/Source/Runtime/Online/BuildPatchServices/Private/Installer/Statistics/DownloadServiceStatistics.cpp
@@ -19,7 +19,7 @@ namespace BuildPatchServices
 		: public IDownloadServiceStatistics
 	{
 	public:
-		FDownloadServiceStatistics(ISpeedRecorder* SpeedRecorder, IDataSizeProvider* DataSizeProvider, IInstallerAnalytics* InstallerAnalytics);
+		FDownloadServiceStatistics(const TSharedPtr<ISpeedRecorder>& SpeedRecorder, const TSharedPtr<IDataSizeProvider>& DataSizeProvider, const TSharedPtr<IInstallerAnalytics>& InstallerAnalytics);
 
 		// IDownloadServiceStat interface begin.
 		virtual void OnDownloadStarted(int32 RequestId, const FString& Uri) override;
@@ -38,9 +38,9 @@ namespace BuildPatchServices
 		// IDownloadServiceStatistics interface end.
 
 	private:
-		ISpeedRecorder* SpeedRecorder; // AddRecord is safe, but SpeedRecorder is ticked
-		IDataSizeProvider* DataSizeProvider;
-		IInstallerAnalytics* InstallerAnalytics;
+		TWeakPtr<ISpeedRecorder> SpeedRecorder; // AddRecord is safe, but SpeedRecorder is ticked
+		TWeakPtr<IDataSizeProvider> DataSizeProvider;
+		TWeakPtr<IInstallerAnalytics> InstallerAnalytics;
 		
 		std::atomic<uint64> TotalBytesReceived;
 		std::atomic<int32> NumSuccessfulDownloads;
@@ -55,7 +55,7 @@ namespace BuildPatchServices
 		mutable FCriticalSection AverageSpeedCriticalSection;
 	};
 
-	FDownloadServiceStatistics::FDownloadServiceStatistics(ISpeedRecorder* InSpeedRecorder, IDataSizeProvider* InDataSizeProvider, IInstallerAnalytics* InInstallerAnalytics)
+	FDownloadServiceStatistics::FDownloadServiceStatistics(const TSharedPtr<ISpeedRecorder>& InSpeedRecorder, const TSharedPtr<IDataSizeProvider>& InDataSizeProvider, const TSharedPtr<IInstallerAnalytics>& InInstallerAnalytics)
 		: SpeedRecorder(InSpeedRecorder)
 		, DataSizeProvider(InDataSizeProvider)
 		, InstallerAnalytics(InInstallerAnalytics)
@@ -93,7 +93,10 @@ namespace BuildPatchServices
 		{
 			TotalBytesReceived += DownloadRecord.SpeedRecord.Size;
 			NumSuccessfulDownloads++;
-			SpeedRecorder->AddRecord(DownloadRecord.SpeedRecord);
+			if (const TSharedPtr<ISpeedRecorder> SpeedRecorderPtr = SpeedRecorder.Pin())
+			{
+				SpeedRecorderPtr->AddRecord(DownloadRecord.SpeedRecord);
+			}
 
 			// update average speed
 			const uint64 Cycles = DownloadRecord.SpeedRecord.CyclesEnd - DownloadRecord.SpeedRecord.CyclesStart;
@@ -104,8 +107,26 @@ namespace BuildPatchServices
 		}
 		else if (DownloadRecord.FailureReason != EHttpFailureReason::Cancelled)
 		{
+			// Ignore deltas as they might not exist. OptimisedDelta enforces the DeltaPolicy and handles errors if deltas are required
+			FString UriFileExtension = FPaths::GetExtension(DownloadRecord.Uri);
+			if (UriFileExtension.Equals(TEXT("delta"), ESearchCase::IgnoreCase) && (DownloadRecord.ResponseCode == EHttpResponseCodes::Forbidden || DownloadRecord.ResponseCode == EHttpResponseCodes::NotFound))
+			{
+				// Parse installed/installing version strings for analytics
+				FString PatchVersion = FPaths::GetBaseFilename(DownloadRecord.Uri);
+				FString BaseVersion = FPaths::GetCleanFilename(FPaths::GetPath(DownloadRecord.Uri));
+
+				if (const TSharedPtr<IInstallerAnalytics> InstallerAnalyticsPtr = InstallerAnalytics.Pin())
+				{
+					InstallerAnalyticsPtr->RecordMissingDelta(DownloadRecord.Uri, BaseVersion, PatchVersion);
+				}
+				return;
+			}
+
 			NumFailedDownloads++;
-			InstallerAnalytics->RecordChunkDownloadError(DownloadRecord.Uri, DownloadRecord.ResponseCode, TEXT("DownloadFail"));
+			if (const TSharedPtr<IInstallerAnalytics> InstallerAnalyticsPtr = InstallerAnalytics.Pin())
+			{
+				InstallerAnalyticsPtr->RecordChunkDownloadError(DownloadRecord.Uri, DownloadRecord.ResponseCode, TEXT("DownloadFail"));
+			}
 		}
 	}
 
@@ -150,8 +171,11 @@ namespace BuildPatchServices
 			}
 		}
 
-		DataSizeProvider->GetDownloadSize(DownloadData, DownloadSizes);
-		check(DownloadSizes.Num() == Result.Num());
+		if (const TSharedPtr<IDataSizeProvider> DataSizeProviderPtr = DataSizeProvider.Pin())
+		{
+			DataSizeProviderPtr->GetDownloadSize(DownloadData, DownloadSizes);
+			check(DownloadSizes.Num() == Result.Num());
+		}
 
 		for (int32 Index = 0; FDownload& Element : Result)
 		{
@@ -198,7 +222,7 @@ namespace BuildPatchServices
 		}
 	}
 
-	IDownloadServiceStatistics* FDownloadServiceStatisticsFactory::Create(ISpeedRecorder* SpeedRecorder, IDataSizeProvider* DataSizeProvider, IInstallerAnalytics* InstallerAnalytics)
+	IDownloadServiceStatistics* FDownloadServiceStatisticsFactory::Create(const TSharedPtr<ISpeedRecorder>& SpeedRecorder, const TSharedPtr<IDataSizeProvider>& DataSizeProvider, const TSharedPtr<IInstallerAnalytics>& InstallerAnalytics)
 	{
 		check(SpeedRecorder != nullptr);
 		check(DataSizeProvider != nullptr);
diff --git a/Engine/Source/Runtime/Online/BuildPatchServices/Private/Installer/Statistics/DownloadServiceStatistics.h b/Engine/Source/Runtime/Online/BuildPatchServices/Private/Installer/Statistics/DownloadServiceStatistics.h
--- a/Engine/Source/Runtime/Online/BuildPatchServices/Private/Installer/Statistics/DownloadServiceStatistics.h
+++ b/Engine/Source/Runtime/Online/BuildPatchServices/Private/Installer/Statistics/DownloadServiceStatistics.h
@@ -73,6 +73,6 @@ namespace BuildPatchServices
 		 * @param InstallerAnalytics    The analytics implementation for reporting the download service events.
 		 * @return the new IDownloadServiceStatistics instance created.
 		 */
-		static IDownloadServiceStatistics* Create(ISpeedRecorder* SpeedRecorder, IDataSizeProvider* DataSizeProvider, IInstallerAnalytics* InstallerAnalytics);
+		static IDownloadServiceStatistics* Create(const TSharedPtr<ISpeedRecorder>& SpeedRecorder, const TSharedPtr<IDataSizeProvider>& DataSizeProvider, const TSharedPtr<IInstallerAnalytics>& InstallerAnalytics);
 	};
 }
\ No newline at end of file
diff --git a/Engine/Source/Runtime/Online/BuildPatchServices/Private/Installer/Statistics/InstallActionStatistics.cpp b/Engine/Source/Runtime/Online/BuildPatchServices/Private/Installer/Statistics/InstallActionStatistics.cpp
new file mode 100644
--- /dev/null
+++ b/Engine/Source/Runtime/Online/BuildPatchServices/Private/Installer/Statistics/InstallActionStatistics.cpp
@@ -0,0 +1,143 @@
+// Copyright Epic Games, Inc. All Rights Reserved.
+
+#include "InstallActionStatistics.h"
+#include "IBuildManifestSet.h"
+#include "Misc/ScopeLock.h"
+#include "Algo/Transform.h"
+
+namespace BuildPatchServices
+{
+	class FInstallActionStatistics
+		: public IInstallActionStatistics
+	{
+		using TActionsArray = TArray<const FBuildPatchInstallerAction*>;
+
+		public:
+		FInstallActionStatistics(const TArray<FBuildPatchInstallerAction>& InInstallerActions)
+			: InstallerActions(InInstallerActions)
+		{
+			InitCachedData();	// Pre - compute data for efficiency.
+		}
+
+		TArray<FInstallActionStats> GetAllStats() const
+		{
+			TArray<FInstallActionStats> Stats;
+			Stats.Reserve(InstallerActions.Num());
+			// Return array order should match installer action order in configuration.
+			for (const FBuildPatchInstallerAction& InstallerAction : InstallerActions)
+			{
+				Stats.Add(GetActionStats(InstallerAction));
+			}
+			return Stats;
+		}
+
+		virtual FInstallActionStats GetActionStats(const FBuildPatchInstallerAction& InstallerAction) const override
+		{
+			const FInstallActionStats* Stats = ActionStatsMap.Find(&InstallerAction);
+			if (Stats != nullptr)
+			{
+				return *Stats;
+			}
+			return {};
+		}
+
+		virtual void OnUsedOptimizedDelta(const FBuildPatchInstallerAction& InstallerAction, bool bInUsedOptimizedDelta, uint64 InHTTPResponseCode) override
+		{
+			FInstallActionStats& Stats = ActionStatsMap.FindOrAdd(&InstallerAction);
+			
+			Stats.UsedOptimizedDelta = bInUsedOptimizedDelta;
+			Stats.OptimizedDeltaHTTPResponseCode = InHTTPResponseCode;
+		}
+
+		virtual void OnChunkRecycleFailure(const FGuid& ChunkId) override
+		{
+			// Find which action(s) chunk belongs to, then increment failure counter for each.
+			TActionsArray* Actions = GuidActionMap.Find(ChunkId);
+			if (Actions != nullptr)
+			{
+				for (const FBuildPatchInstallerAction* InstallerAction: *Actions)
+				{
+					ActionStatsMap.FindOrAdd(InstallerAction).NumChunksRecycleFailures++;
+				}
+			}
+		}
+
+		virtual void OnNumLocalChunksExpected(const FBuildPatchInstallerAction& InstallerAction, uint64 InNumLocalChunksExpected) override
+		{
+			ActionStatsMap.FindOrAdd(&InstallerAction).NumLocalChunksExpected = InNumLocalChunksExpected;
+		}
+
+		virtual void OnNumLocalChunksAvailable(const FBuildPatchInstallerAction& InstallerAction, uint64 InNumLocalChunksAvailable) override
+		{
+			ActionStatsMap.FindOrAdd(&InstallerAction).NumLocalChunksAvailable = InNumLocalChunksAvailable;
+		}
+
+		virtual void CalculateChunkStatistics(const TSet<FGuid>& InstallChunksAvailable) override
+		{
+			for (const FBuildPatchInstallerAction& Action : InstallerActions)
+			{
+				// Calculate expected and available local chunks (recyclable from existing installation).
+				// Fresh installs don't have anything to recycle - process updates and repairs only.
+				if (Action.IsUpdate() || Action.IsRepair())
+				{
+					const FBuildPatchAppManifest& CurrentManifest = Action.GetCurrentManifest();
+					const FBuildPatchAppManifest& DestinationManifest = Action.GetInstallManifest();
+
+					TSet<FGuid> CurrentChunks;
+					// Get chunks from current installation (what we should have for all files in current manifest).
+					TSet<FString> CurrentFilesSet(CurrentManifest.GetBuildFileList(Action.GetInstallTags()));
+
+					CurrentManifest.GetChunksRequiredForFiles(CurrentFilesSet, CurrentChunks);
+
+					// Get chunks from destination (what we need) from cache.
+					TSet<FGuid>& DestinationChunks = ChunkCacheMap.FindOrAdd(&Action);
+
+					// Chunks we should be able recycle.
+					OnNumLocalChunksExpected(Action, DestinationChunks.Intersect(CurrentChunks).Num());
+
+					// Chunks that are actually available for recycling.
+					OnNumLocalChunksAvailable(Action, DestinationChunks.Intersect(InstallChunksAvailable).Num());
+				}
+			}
+		}
+
+	private:
+		const TArray<FBuildPatchInstallerAction>& InstallerActions;
+		TMap<const FBuildPatchInstallerAction*, FInstallActionStats> ActionStatsMap;
+		TMap<FGuid, TActionsArray> GuidActionMap;
+		TMap<const FBuildPatchInstallerAction*, TSet<FGuid>> ChunkCacheMap;
+
+		void InitCachedData()
+		{
+			for (const FBuildPatchInstallerAction& InstallerAction : InstallerActions)
+			{
+				const FBuildPatchAppManifestPtr InstallManifest = InstallerAction.TryGetSharedInstallManifest();
+				if (!InstallManifest.IsValid())
+				{
+					continue;
+				}
+
+				TSet<FGuid> Chunks;
+				TSet<FString> CurrentFilesSet(InstallManifest->GetBuildFileList(InstallerAction.GetInstallTags()));
+
+				InstallManifest->GetChunksRequiredForFiles(CurrentFilesSet, Chunks);
+
+				// Store chunks in cache - used in CalculateChunkStatistics.
+				ChunkCacheMap.Add(&InstallerAction, Chunks);
+
+				// Build ChunkId-Actions map - used in OnChunkRecycleFailure
+				// to find which actions failed chunk belongs to.
+				for (const FGuid& ChunkId : Chunks)
+				{
+					TActionsArray& Actions = GuidActionMap.FindOrAdd(ChunkId);
+					Actions.AddUnique(&InstallerAction);
+				}
+			}
+		}
+	};
+
+	IInstallActionStatistics* FInstallActionStatisticsFactory::Create(const TArray<FBuildPatchInstallerAction>& InstallerActions)
+	{
+		return new FInstallActionStatistics(InstallerActions);
+	}
+}
\ No newline at end of file
diff --git a/Engine/Source/Runtime/Online/BuildPatchServices/Private/Installer/Statistics/InstallActionStatistics.h b/Engine/Source/Runtime/Online/BuildPatchServices/Private/Installer/Statistics/InstallActionStatistics.h
new file mode 100644
--- /dev/null
+++ b/Engine/Source/Runtime/Online/BuildPatchServices/Private/Installer/Statistics/InstallActionStatistics.h
@@ -0,0 +1,73 @@
+// Copyright Epic Games, Inc. All Rights Reserved.
+#pragma once
+
+#include "CoreMinimal.h"
+#include "Interfaces/IBuildInstaller.h"
+
+namespace BuildPatchServices
+{
+	struct FBuildPatchInstallerAction;
+
+	class IInstallActionStatistics
+	{
+	public:
+		virtual ~IInstallActionStatistics() = default;
+
+		/**
+		 * Get all installer actions stats. Order matches the order of installer actions in configuration.
+		 * @return The list of installer actions stats.
+		 */
+		virtual TArray<FInstallActionStats> GetAllStats() const = 0;
+
+		/**
+		 * Return statistics for one action.
+		 * @param InstallerAction The installer action to query.
+		 * @return The statistics for queried action.
+		 */
+		virtual FInstallActionStats GetActionStats(const FBuildPatchInstallerAction& InstallerAction) const = 0;
+
+		/**
+		* Calculate expected and available statistics for all actions,
+		* based on action manifests and chunks available from chunk source.
+		* @param InstallChunksAvailable The list of chunks available.
+		*/
+		virtual void CalculateChunkStatistics(const TSet<FGuid>& InstallChunksAvailable) = 0;
+
+		/**
+		 * Set whether the action successfully downloaded and used delta optimization file.
+		 * @param InstallerAction The installer action to set.
+		 * @param bInUsedOptimizedDelta The optimized delta flag.
+		 */
+		virtual void OnUsedOptimizedDelta(const FBuildPatchInstallerAction& InstallerAction, bool bInUsedOptimizedDelta, uint64 InHTTPResponseCode) = 0;
+
+		/**
+		 * Set the number of local chunks expected for the action.
+		 * @param InstallerAction The installer action to set.
+		 * @param InNumLocalChunksExpected The number of local chunks expected.
+		 */
+		virtual void OnNumLocalChunksExpected(const FBuildPatchInstallerAction& InstallerAction, uint64 InNumLocalChunksExpected) = 0;
+
+		/**
+		 * Set the number of local chunks available for the action.
+		 * @param InstallerAction The installer action to set.
+		 * @param InNumLocalChunksAvailable The number of local chunks available.
+		 */
+		virtual void OnNumLocalChunksAvailable(const FBuildPatchInstallerAction& InstallerAction, uint64 InNumLocalChunksAvailable) = 0;
+
+		/**
+		 * Increment the number of chunks that failed to recycle.
+		 * @param ChunkId The id of a chunk that was not recycled.
+		 */
+		virtual void OnChunkRecycleFailure(const FGuid& ChunkId) = 0;
+	};
+
+	class FInstallActionStatisticsFactory
+	{
+		public:
+		/**
+		 * Create an install action statistics instance.
+		 * @return The created instance.
+		 */
+		static IInstallActionStatistics* Create(const TArray<FBuildPatchInstallerAction> &InstallerActions);
+	};
+}
\ No newline at end of file
diff --git a/Engine/Source/Runtime/Online/BuildPatchServices/Private/Installer/Statistics/InstallChunkSourceStatistics.cpp b/Engine/Source/Runtime/Online/BuildPatchServices/Private/Installer/Statistics/InstallChunkSourceStatistics.cpp
--- a/Engine/Source/Runtime/Online/BuildPatchServices/Private/Installer/Statistics/InstallChunkSourceStatistics.cpp
+++ b/Engine/Source/Runtime/Online/BuildPatchServices/Private/Installer/Statistics/InstallChunkSourceStatistics.cpp
@@ -6,6 +6,7 @@
 #include "Common/SpeedRecorder.h"
 #include "Common/StatsCollector.h"
 #include "Installer/InstallerAnalytics.h"
+#include "InstallActionStatistics.h"
 #include "Installer/Statistics/FileOperationTracker.h"
 
 namespace BuildPatchServices
@@ -14,7 +15,7 @@ namespace BuildPatchServices
 		: public IInstallChunkSourceStatistics
 	{
 	public:
-		FInstallChunkSourceStatistics(ISpeedRecorder* SpeedRecorder, IInstallerAnalytics* InstallerAnalytics, IFileOperationTracker* FileOperationTracker);
+		FInstallChunkSourceStatistics(ISpeedRecorder* SpeedRecorder, IInstallerAnalytics* InstallerAnalytics, IInstallActionStatistics* InInstallerActionStatistics, IFileOperationTracker* FileOperationTracker);
 		~FInstallChunkSourceStatistics();
 
 		// IInstallChunkSourceStat interface begin.
@@ -32,6 +33,7 @@ namespace BuildPatchServices
 	private:
 		ISpeedRecorder* SpeedRecorder;
 		IInstallerAnalytics* InstallerAnalytics;
+		IInstallActionStatistics* InstallerActionStatistics;
 		IFileOperationTracker* FileOperationTracker;
 		FThreadSafeInt32 NumSuccessfulLoads;
 		FThreadSafeInt32 NumFailedLoads;
@@ -40,9 +42,10 @@ namespace BuildPatchServices
 		FThreadSafeInt64 LastReadCycles;
 	};
 
-	FInstallChunkSourceStatistics::FInstallChunkSourceStatistics(ISpeedRecorder* InSpeedRecorder, IInstallerAnalytics* InInstallerAnalytics, IFileOperationTracker* InFileOperationTracker)
+	FInstallChunkSourceStatistics::FInstallChunkSourceStatistics(ISpeedRecorder* InSpeedRecorder, IInstallerAnalytics* InInstallerAnalytics, IInstallActionStatistics* InInstallerActionStatistics, IFileOperationTracker* InFileOperationTracker)
 		: SpeedRecorder(InSpeedRecorder)
 		, InstallerAnalytics(InInstallerAnalytics)
+		, InstallerActionStatistics(InInstallerActionStatistics)
 		, FileOperationTracker(InFileOperationTracker)
 		, NumSuccessfulLoads(0)
 		, NumFailedLoads(0)
@@ -72,6 +75,7 @@ namespace BuildPatchServices
 		{
 			FileOperationTracker->OnDataStateUpdate(ChunkId, EFileOperationState::PendingRemoteCloudData);
 			InstallerAnalytics->RecordChunkCacheError(ChunkId, TEXT(""), FPlatformMisc::GetLastError(), TEXT("InstallChunkSourceLoad"), ToString(Result));
+			InstallerActionStatistics->OnChunkRecycleFailure(ChunkId);
 			NumFailedLoads.Increment();
 		}
 		SpeedRecorder->AddRecord(Record);
@@ -101,11 +105,11 @@ namespace BuildPatchServices
 		return bIsReading || FStatsCollector::CyclesToSeconds(CyclesSinceRead) < 0.05;
 	}
 
-	IInstallChunkSourceStatistics* FInstallChunkSourceStatisticsFactory::Create(ISpeedRecorder* SpeedRecorder, IInstallerAnalytics* InstallerAnalytics, IFileOperationTracker* FileOperationTracker)
+	IInstallChunkSourceStatistics* FInstallChunkSourceStatisticsFactory::Create(ISpeedRecorder* SpeedRecorder, IInstallerAnalytics* InstallerAnalytics, IInstallActionStatistics* InInstallerActionStatistics, IFileOperationTracker* FileOperationTracker)
 	{
 		check(SpeedRecorder != nullptr);
 		check(InstallerAnalytics != nullptr);
 		check(FileOperationTracker != nullptr);
-		return new FInstallChunkSourceStatistics(SpeedRecorder, InstallerAnalytics, FileOperationTracker);
+		return new FInstallChunkSourceStatistics(SpeedRecorder, InstallerAnalytics, InInstallerActionStatistics, FileOperationTracker);
 	}
 };
\ No newline at end of file
diff --git a/Engine/Source/Runtime/Online/BuildPatchServices/Private/Installer/Statistics/InstallChunkSourceStatistics.h b/Engine/Source/Runtime/Online/BuildPatchServices/Private/Installer/Statistics/InstallChunkSourceStatistics.h
--- a/Engine/Source/Runtime/Online/BuildPatchServices/Private/Installer/Statistics/InstallChunkSourceStatistics.h
+++ b/Engine/Source/Runtime/Online/BuildPatchServices/Private/Installer/Statistics/InstallChunkSourceStatistics.h
@@ -7,6 +7,7 @@ namespace BuildPatchServices
 {
 	class ISpeedRecorder;
 	class IInstallerAnalytics;
+	class IInstallActionStatistics;
 	class IFileOperationTracker;
 
 	/**
@@ -50,6 +51,6 @@ namespace BuildPatchServices
 		 * @param FileOperationTracker  The file operation tracker which will be used to update data states.
 		 * @return the new IInstallChunkSourceStatistics instance created.
 		 */
-		static IInstallChunkSourceStatistics* Create(ISpeedRecorder* SpeedRecorder, IInstallerAnalytics* InstallerAnalytics, IFileOperationTracker* FileOperationTracker);
+		static IInstallChunkSourceStatistics* Create(ISpeedRecorder* SpeedRecorder, IInstallerAnalytics* InstallerAnalytics, IInstallActionStatistics* InInstallerActionStatistics, IFileOperationTracker* FileOperationTracker);
 	};
 }
\ No newline at end of file
diff --git a/Engine/Source/Runtime/Online/BuildPatchServices/Private/Tests/Fake/InstallerError.fake.h b/Engine/Source/Runtime/Online/BuildPatchServices/Private/Tests/Fake/InstallerError.fake.h
--- a/Engine/Source/Runtime/Online/BuildPatchServices/Private/Tests/Fake/InstallerError.fake.h
+++ b/Engine/Source/Runtime/Online/BuildPatchServices/Private/Tests/Fake/InstallerError.fake.h
@@ -61,8 +61,12 @@ namespace BuildPatchServices
 			Delegates.Remove(Handle);
 		}
 
-		virtual void Reset() override
+		virtual void ResetError(bool bPreserveCancelled) override
 		{
+			if (bPreserveCancelled && bCanRetry)
+			{
+				return;
+			}
 			bHasError = false;
 			bIsCancelled = false;
 			bCanRetry = true;
diff --git a/Engine/Source/Runtime/Online/BuildPatchServices/Private/Tests/Mock/InstallerAnalytics.mock.h b/Engine/Source/Runtime/Online/BuildPatchServices/Private/Tests/Mock/InstallerAnalytics.mock.h
--- a/Engine/Source/Runtime/Online/BuildPatchServices/Private/Tests/Mock/InstallerAnalytics.mock.h
+++ b/Engine/Source/Runtime/Online/BuildPatchServices/Private/Tests/Mock/InstallerAnalytics.mock.h
@@ -14,6 +14,7 @@ namespace BuildPatchServices
 	{
 	public:
 		typedef TTuple<double, FString, int32,  FString> FRecordChunkDownloadError;
+		typedef TTuple<double, FString, FString, FString> FRecordMissingDelta;
 		typedef TTuple<double, FString, double, double, double, double> FRecordChunkDownloadAborted;
 		typedef TTuple<double, FGuid, FString, int32, FString, FString> FRecordChunkCacheError;
 		typedef TTuple<double, FString, int32, FString> FRecordConstructionError;
@@ -26,6 +27,11 @@ namespace BuildPatchServices
 			RxRecordChunkDownloadError.Emplace(FStatsCollector::GetSeconds(), ChunkUrl, ResponseCode, ErrorString);
 		}
 
+		virtual void RecordMissingDelta(const FString& DeltaUrl, const FString& BaseVersion, const FString& PatchVersion) override
+		{
+			RxRecordMissingDelta.Emplace(FStatsCollector::GetSeconds(), DeltaUrl, BaseVersion, PatchVersion);
+		}
+
 		virtual void RecordChunkDownloadAborted(const FString& ChunkUrl, double ChunkTime, double ChunkMean, double ChunkStd, double BreakingPoint) override
 		{
 			RxRecordChunkDownloadAborted.Emplace(FStatsCollector::GetSeconds(), ChunkUrl, ChunkTime, ChunkMean, ChunkStd, BreakingPoint);
@@ -58,6 +64,7 @@ namespace BuildPatchServices
 
 	public:
 		TArray<FRecordChunkDownloadError> RxRecordChunkDownloadError;
+		TArray<FRecordMissingDelta> RxRecordMissingDelta;
 		TArray<FRecordChunkDownloadAborted> RxRecordChunkDownloadAborted;
 		TArray<FRecordChunkCacheError> RxRecordChunkCacheError;
 		TArray<FRecordConstructionError> RxRecordConstructionError;
diff --git a/Engine/Source/Runtime/Online/BuildPatchServices/Public/BuildPatchSettings.h b/Engine/Source/Runtime/Online/BuildPatchServices/Public/BuildPatchSettings.h
--- a/Engine/Source/Runtime/Online/BuildPatchServices/Public/BuildPatchSettings.h
+++ b/Engine/Source/Runtime/Online/BuildPatchServices/Public/BuildPatchSettings.h
@@ -419,8 +419,12 @@ namespace BuildPatchServices
 		FString UninstallActionArgs;
 		// The maximum age (in days) of existing data files which can be reused in this build.
 		float DataAgeThreshold = TNumericLimits<float>::Max();
-		// Indicates whether data age threshold should be honored. If false, ALL data files can be reused.
+		// Indicates whether data age threshold should be honored. If false, ALL data files can be reused taking into
+		// account the ReuseManifestPredicate.
 		bool bShouldHonorReuseThreshold = false;
+		// A callable that takes a const reference to an IBuildManifest and returns true to reuse the manifest
+		// and its chunks, or false to fully ignore it.
+		TFunction<bool(const IBuildManifest&)> ReuseManifestPredicate = [](const IBuildManifest&) { return true; };
 		// The chunk window size to be used when saving out new data.
 		uint32 OutputChunkWindowSize = 1024 * 1024;
 		// Indicates whether any window size chunks should be matched, rather than just out output window size.
@@ -528,6 +532,8 @@ namespace BuildPatchServices
 		// the patcher to select matches from all of the aliased tags.
 		TMap<FString, FString> ManifestATagAliases;
 		TMap<FString, FString> ManifestBTagAliases;
+		TSet<FString> ManifestAIgnoreTags;
+		TSet<FString> ManifestBIgnoreTags;
 		
 		// The custom optimized delta filename to use for generating this optimized delta. If used, all BPS modes
 		// that request an optimized delta will need the same trailer passed in.
diff --git a/Engine/Source/Runtime/Online/BuildPatchServices/Public/Interfaces/IBuildInstaller.h b/Engine/Source/Runtime/Online/BuildPatchServices/Public/Interfaces/IBuildInstaller.h
--- a/Engine/Source/Runtime/Online/BuildPatchServices/Public/Interfaces/IBuildInstaller.h
+++ b/Engine/Source/Runtime/Online/BuildPatchServices/Public/Interfaces/IBuildInstaller.h
@@ -61,7 +61,9 @@ enum class EBuildPatchInstallError
 	CustomUninstallActionError = 12,
 
 	// Used to help verify logic
-	NumInstallErrors
+	InvalidOrMax,
+	// Deprecated value. Use InvalidOrMax instad of NumInstallErrors
+	NumInstallErrors = InvalidOrMax,
 };
 
 /**
@@ -120,6 +122,19 @@ struct FCloudDirectoryStats
 	uint32 ErrorCount = 0;
 };
 
+/**
+ * A struct to hold stats for a single install action.
+ */
+struct FInstallActionStats
+{
+public:
+	bool UsedOptimizedDelta = false;
+	uint64 OptimizedDeltaHTTPResponseCode = 0;
+	uint64 NumLocalChunksExpected = 0;
+	uint64 NumLocalChunksAvailable = 0;
+	uint64 NumChunksRecycleFailures = 0;
+};
+
 /**
  * A struct to hold stats for the build process.
  */
@@ -185,6 +200,7 @@ struct FBuildInstallStats
 		, ProcessRequiredDiskSpace(0)
 		, ProcessAvailableDiskSpace(0)
 		, CloudDirectoryStats()
+		, ActionsStats()
 	{}
 
 	// The total number of files in the build.
@@ -324,6 +340,8 @@ struct FBuildInstallStats
 	uint64 ProcessAvailableDiskSpace;
 	// Stats tracked per cloud directory.
 	TArray<FCloudDirectoryStats> CloudDirectoryStats;
+	// Stats tracked per action. Array order matches actions order in configuration.
+	TArray<FInstallActionStats> ActionsStats;
 };
 
 namespace BuildPatchServices
@@ -545,7 +563,7 @@ public:
 	virtual void SetPerFileSubdirectories(TMap<FString, FString>&& InOldPerFileSubdirectories, TMap<FString, FString>&& InNewPerFileSubdirectories) = 0;
 };
 
-static_assert((uint32)EBuildPatchInstallError::NumInstallErrors == 13, "Please add support for the extra values to the Lex functions below.");
+static_assert((uint32)EBuildPatchInstallError::InvalidOrMax == 13, "Please add support for the extra values to the Lex functions below.");
 
 inline const TCHAR* LexToString(EBuildPatchInstallError Error)
 {
@@ -593,7 +611,7 @@ inline void LexFromString(EBuildPatchInstallError& Error, const TCHAR* Buffer)
 	RETURN_IF_EQUAL(OutOfDiskSpace);
 	RETURN_IF_EQUAL(CustomUninstallActionError);
 	// Did not match
-	Error = EBuildPatchInstallError::NumInstallErrors;
+	Error = EBuildPatchInstallError::InvalidOrMax;
 	return;
 #undef RETURN_IF_EQUAL
 }
diff --git a/Engine/Source/Runtime/Online/BuildPatchServices/Public/Interfaces/IBuildManifest.h b/Engine/Source/Runtime/Online/BuildPatchServices/Public/Interfaces/IBuildManifest.h
--- a/Engine/Source/Runtime/Online/BuildPatchServices/Public/Interfaces/IBuildManifest.h
+++ b/Engine/Source/Runtime/Online/BuildPatchServices/Public/Interfaces/IBuildManifest.h
@@ -119,6 +119,13 @@ public:
 	 */
 	virtual bool IsManifestEncrypted() const = 0;
 
+	/**
+	 * Returns whether the manifest and all it chunks were originally encrypted when first loaded.
+	 * This returns true even if the manifest has since been decrypted in memory.
+	 * @return    true if the manifest was originally encrypted
+	 */
+	virtual bool IsOriginallyFullyEncrypted() const = 0;
+
 	/**
 	 * Get all necessary secret key IDs.
 	 * @return		full set of necessary secret ids
